Checkpoint: Experimental Audio Widget Plugin
- First UMG pass on Audio Knob Large/Small
- First UMG pass on Audio Button Toggle
- First UMG pass on Audio Button Matrix
- Audio Fader started (initial pass on layout, logic not hooked up)

diff --git a/Engine/Plugins/Runtime/AudioWidgets/AudioWidgets.uplugin b/Engine/Plugins/Runtime/AudioWidgets/AudioWidgets.uplugin
new file mode 100644
--- /dev/null
+++ b/Engine/Plugins/Runtime/AudioWidgets/AudioWidgets.uplugin
@@ -0,0 +1,24 @@
+{
+	"FileVersion": 3,
+	"Version": 1,
+	"VersionName": "1.0",
+	"FriendlyName": "AudioWidgets",
+	"Description": "Collection of widgets tailored to interacting with audio-related data and systems.",
+	"Category": "Audio",
+	"CreatedBy": "Epic Games, Inc.",
+	"CreatedByURL": "http://epicgames.com",
+	"DocsURL": "",
+	"MarketplaceURL": "",
+	"SupportURL": "",
+	"CanContainContent": true,
+	"IsBetaVersion": true,
+	"IsExperimentalVersion": false,
+	"Installed": false,
+	"Modules": [
+		{
+			"Name": "AudioWidgets",
+			"Type": "Runtime",
+			"LoadingPhase": "Default"
+		}
+	]
+}
\ No newline at end of file
diff --git a/Engine/Plugins/Runtime/AudioWidgets/Source/AudioWidgets/AudioWidgets.Build.cs b/Engine/Plugins/Runtime/AudioWidgets/Source/AudioWidgets/AudioWidgets.Build.cs
new file mode 100644
--- /dev/null
+++ b/Engine/Plugins/Runtime/AudioWidgets/Source/AudioWidgets/AudioWidgets.Build.cs
@@ -0,0 +1,29 @@
+// Copyright Epic Games, Inc. All Rights Reserved.
+
+using UnrealBuildTool;
+
+public class AudioWidgets : ModuleRules
+{
+	public AudioWidgets(ReadOnlyTargetRules Target) : base(Target)
+	{
+		PCHUsage = ModuleRules.PCHUsageMode.UseExplicitOrSharedPCHs;
+
+		PublicDependencyModuleNames.AddRange(
+		new string[]
+			{
+				"Core",
+			}
+		);
+
+		PrivateDependencyModuleNames.AddRange(
+		new string[]
+			{
+				"CoreUObject",
+				"Engine",
+				"Slate",
+				"SlateCore",
+				"AdvancedWidgets"
+			}
+		);
+	}
+}
diff --git a/Engine/Plugins/Runtime/AudioWidgets/Source/AudioWidgets/Private/AudioWidgets.cpp b/Engine/Plugins/Runtime/AudioWidgets/Source/AudioWidgets/Private/AudioWidgets.cpp
new file mode 100644
--- /dev/null
+++ b/Engine/Plugins/Runtime/AudioWidgets/Source/AudioWidgets/Private/AudioWidgets.cpp
@@ -0,0 +1,20 @@
+// Copyright Epic Games, Inc. All Rights Reserved.
+
+#include "AudioWidgets.h"
+
+#include "AdvancedWidgetsModule.h"
+
+#define LOCTEXT_NAMESPACE "FAudioWidgetsModule"
+
+void FAudioWidgetsModule::StartupModule()
+{
+	FModuleManager::Get().LoadModuleChecked<FAdvancedWidgetsModule>("AdvancedWidgets");
+}
+
+void FAudioWidgetsModule::ShutdownModule()
+{
+}
+
+#undef LOCTEXT_NAMESPACE
+	
+IMPLEMENT_MODULE(FAudioWidgetsModule, AudioWidgets)
\ No newline at end of file
diff --git a/Engine/Plugins/Runtime/AudioWidgets/Source/AudioWidgets/Public/AudioWidgets.h b/Engine/Plugins/Runtime/AudioWidgets/Source/AudioWidgets/Public/AudioWidgets.h
new file mode 100644
--- /dev/null
+++ b/Engine/Plugins/Runtime/AudioWidgets/Source/AudioWidgets/Public/AudioWidgets.h
@@ -0,0 +1,15 @@
+// Copyright Epic Games, Inc. All Rights Reserved.
+
+#pragma once
+
+#include "CoreMinimal.h"
+#include "Modules/ModuleManager.h"
+
+class FAudioWidgetsModule : public IModuleInterface
+{
+public:
+
+	/** IModuleInterface implementation */
+	virtual void StartupModule() override;
+	virtual void ShutdownModule() override;
+};
