Adding a new audio meter slate widget

diff --git a/Engine/Plugins/Runtime/AudioSynesthesia/Source/AudioSynesthesia/Classes/Meter.h b/Engine/Plugins/Runtime/AudioSynesthesia/Source/AudioSynesthesia/Classes/Meter.h
--- a/Engine/Plugins/Runtime/AudioSynesthesia/Source/AudioSynesthesia/Classes/Meter.h
+++ b/Engine/Plugins/Runtime/AudioSynesthesia/Source/AudioSynesthesia/Classes/Meter.h
@@ -42,14 +42,6 @@ public:
 	UPROPERTY(EditAnywhere, BlueprintReadOnly, Category = AudioAnalyzer, meta = (ClampMin = "0"))
 	int32 MeterReleaseTime = 300;
 
-	/** Peak attack time, in milliseconds */
-	UPROPERTY(EditAnywhere, BlueprintReadOnly, Category = AudioAnalyzer, meta = (ClampMin = "0"))
-	int32 PeakAttackTime = 300;
-
-	/** Peak release time, in milliseconds */
-	UPROPERTY(EditAnywhere, BlueprintReadOnly, Category = AudioAnalyzer, meta = (ClampMin = "0"))
-	int32 PeakReleaseTime = 300;
-
 	/** Peak hold time, in milliseconds */
 	UPROPERTY(EditAnywhere, BlueprintReadOnly, Category = AudioAnalyzer, meta = (ClampMin = "0"))
 	int32 PeakHoldTime = 100;
diff --git a/Engine/Plugins/Runtime/AudioSynesthesia/Source/AudioSynesthesia/Private/Meter.cpp b/Engine/Plugins/Runtime/AudioSynesthesia/Source/AudioSynesthesia/Private/Meter.cpp
--- a/Engine/Plugins/Runtime/AudioSynesthesia/Source/AudioSynesthesia/Private/Meter.cpp
+++ b/Engine/Plugins/Runtime/AudioSynesthesia/Source/AudioSynesthesia/Private/Meter.cpp
@@ -5,6 +5,7 @@
 #include "InterpolateSorted.h"
 #include "AudioSynesthesiaLog.h"
 #include "DSP/EnvelopeFollower.h"
+#include "DSP/Dsp.h"
 
 TUniquePtr<Audio::IAnalyzerSettings> UMeterSettings::GetSettings(const int32 InSampleRate, const int32 InNumChannels) const
 {
@@ -30,8 +31,6 @@ TUniquePtr<Audio::IAnalyzerSettings> UMeterSettings::GetSettings(const int32 InS
 
 	Settings->MeterAttackTime = MeterAttackTime;
 	Settings->MeterReleaseTime = MeterReleaseTime;
-	Settings->PeakAttackTime = PeakAttackTime;
-	Settings->PeakReleaseTime = PeakReleaseTime;
 	Settings->PeakHoldTime = PeakHoldTime;
 	Settings->ClippingThreshold = ClippingThreshold;
 
@@ -75,9 +74,9 @@ static TArray<FMeterResults> ConvertToBlueprintResults(UMeterSettings* Settings,
 	{
 		FMeterResults NewResults;
 
-		NewResults.MeterValue = MeterEntry.MeterValue;
-		NewResults.PeakValue = MeterEntry.PeakValue;
-		NewResults.ClippingValue = MeterEntry.ClippingValue;
+		NewResults.MeterValue = Audio::ConvertToDecibels(MeterEntry.MeterValue);
+		NewResults.PeakValue = Audio::ConvertToDecibels(MeterEntry.PeakValue);
+		NewResults.ClippingValue = Audio::ConvertToDecibels(MeterEntry.ClippingValue);
 		NewResults.NumSamplesClipping = MeterEntry.NumSamplesClipping;
 		NewResults.TimeSeconds = MeterEntry.Timestamp;
 
diff --git a/Engine/Plugins/Runtime/AudioSynesthesia/Source/AudioSynesthesiaCore/Private/MeterAnalyzer.cpp b/Engine/Plugins/Runtime/AudioSynesthesia/Source/AudioSynesthesiaCore/Private/MeterAnalyzer.cpp
--- a/Engine/Plugins/Runtime/AudioSynesthesia/Source/AudioSynesthesiaCore/Private/MeterAnalyzer.cpp
+++ b/Engine/Plugins/Runtime/AudioSynesthesia/Source/AudioSynesthesiaCore/Private/MeterAnalyzer.cpp
@@ -24,7 +24,6 @@ namespace Audio
 		for (int32 ChannelIndex = 0; ChannelIndex < NumChannels; ++ChannelIndex)
 		{
 			MeterEnvelopeFollowers.Add(MakeUnique<FEnvelopeFollower>(InSampleRate, (float)InSettings.MeterAttackTime, (float)InSettings.MeterReleaseTime, InSettings.MeterPeakMode, true));
-			PeakEnvelopeFollowers.Add(MakeUnique<FEnvelopeFollower>(InSampleRate, (float)InSettings.PeakAttackTime, (float)InSettings.PeakReleaseTime, InSettings.MeterPeakMode, true));
 		}
 	}
 
@@ -49,10 +48,7 @@ namespace Audio
 				}
 
 				TUniquePtr<FEnvelopeFollower>& MeterEnvelopeFollower = MeterEnvelopeFollowers[ChannelIndex];
-				MeterEnvelopeFollower->ProcessAudio(Sample);
-
-				TUniquePtr<FEnvelopeFollower>& PeakEnvelopeFollower = PeakEnvelopeFollowers[ChannelIndex];
-				float NewPeakEnvValue = PeakEnvelopeFollower->ProcessAudio(Sample);
+				float NewPeakEnvValue = MeterEnvelopeFollower->ProcessAudio(Sample);
 
 				FPeakEnvelopeData& PeakEnvData = PeakEnvDataPerChannel[ChannelIndex];
 				FPeakData& PeakData = PeakDataPerChannel[ChannelIndex];
diff --git a/Engine/Plugins/Runtime/AudioSynesthesia/Source/AudioSynesthesiaCore/Public/MeterAnalyzer.h b/Engine/Plugins/Runtime/AudioSynesthesia/Source/AudioSynesthesiaCore/Public/MeterAnalyzer.h
--- a/Engine/Plugins/Runtime/AudioSynesthesia/Source/AudioSynesthesiaCore/Public/MeterAnalyzer.h
+++ b/Engine/Plugins/Runtime/AudioSynesthesia/Source/AudioSynesthesiaCore/Public/MeterAnalyzer.h
@@ -17,12 +17,6 @@ namespace Audio
 		// Envelope follower release time in milliseconds
 		int32 MeterReleaseTime = 300;
 		
-		// Peak detector attack time in milliseconds
-		int32 PeakAttackTime = 300;
-
-		// Peak detector release time in milliseconds
-		int32 PeakReleaseTime = 300;
-
 		// Peak detector hold time in milliseconds
 		int32 PeakHoldTime = 100;
 
@@ -60,7 +54,6 @@ namespace Audio
 	protected:
 		// Envelope follower per channel
 		TArray<TUniquePtr<FEnvelopeFollower>> MeterEnvelopeFollowers;
-		TArray<TUniquePtr<FEnvelopeFollower>> PeakEnvelopeFollowers;
 
 		// Per-channel clipping data
 		struct FClippingData
diff --git a/Engine/Plugins/Runtime/AudioWidgets/Source/AudioWidgets/AudioWidgets.Build.cs b/Engine/Plugins/Runtime/AudioWidgets/Source/AudioWidgets/AudioWidgets.Build.cs
--- a/Engine/Plugins/Runtime/AudioWidgets/Source/AudioWidgets/AudioWidgets.Build.cs
+++ b/Engine/Plugins/Runtime/AudioWidgets/Source/AudioWidgets/AudioWidgets.Build.cs
@@ -12,6 +12,7 @@ public class AudioWidgets : ModuleRules
 		new string[]
 			{
 				"Core",
+				"UMG"
 			}
 		);
 
@@ -22,7 +23,9 @@ public class AudioWidgets : ModuleRules
 				"Engine",
 				"Slate",
 				"SlateCore",
-				"AdvancedWidgets"
+				"InputCore",
+				"AdvancedWidgets",
+				"SignalProcessing"
 			}
 		);
 	}
diff --git a/Engine/Plugins/Runtime/AudioWidgets/Source/AudioWidgets/Private/AudioMeter.cpp b/Engine/Plugins/Runtime/AudioWidgets/Source/AudioWidgets/Private/AudioMeter.cpp
new file mode 100644
--- /dev/null
+++ b/Engine/Plugins/Runtime/AudioWidgets/Source/AudioWidgets/Private/AudioMeter.cpp
@@ -0,0 +1,169 @@
+// Copyright Epic Games, Inc. All Rights Reserved.
+
+#include "AudioMeter.h"
+#include "SAudioMeter.h"
+#include "Widgets/DeclarativeSyntaxSupport.h"
+
+
+#define LOCTEXT_NAMESPACE "AUDIO_UMG"
+
+/////////////////////////////////////////////////////
+// UAudioMeter
+
+static FAudioMeterStyle* DefaultAudioMeterStyle = nullptr;
+
+UAudioMeter2::UAudioMeter2(const FObjectInitializer& ObjectInitializer)
+	: Super(ObjectInitializer)
+{
+	Orientation = EOrientation::Orient_Vertical;
+
+	BackgroundColor = FLinearColor::Black;
+	MeterBackgroundColor = FLinearColor::Gray;
+	MeterValueColor = FLinearColor::Green;
+	MeterPeakColor = FLinearColor::Blue;
+	MeterClippingColor = FLinearColor::Red;
+	MeterScaleColor = FLinearColor::Gray;
+	MeterScaleLabelColor = FLinearColor::White;
+
+	if (DefaultAudioMeterStyle == nullptr)
+	{
+		// HACK: THIS SHOULD NOT COME FROM CORESTYLE AND SHOULD INSTEAD BE DEFINED BY ENGINE TEXTURES/PROJECT SETTINGS
+		DefaultAudioMeterStyle = new FAudioMeterStyle(FCoreStyle::Get().GetWidgetStyle<FAudioMeterStyle>("AudioMeter"));
+
+		// Unlink UMG default colors from the editor settings colors.
+		DefaultAudioMeterStyle->UnlinkColors();
+	}
+
+	WidgetStyle = *DefaultAudioMeterStyle;
+
+#if WITH_EDITORONLY_DATA
+	AccessibleBehavior = ESlateAccessibleBehavior::NotAccessible;
+	bCanChildrenBeAccessible = false;
+#endif
+}
+
+TSharedRef<SWidget> UAudioMeter2::RebuildWidget()
+{
+	MyAudioMeter = SNew(SAudioMeter)
+		.Style(&WidgetStyle);
+
+	return MyAudioMeter.ToSharedRef();
+}
+
+void UAudioMeter2::SynchronizeProperties()
+{
+	Super::SynchronizeProperties();
+
+	MyAudioMeter->SetOrientation(Orientation);
+
+	MyAudioMeter->SetBackgroundColor(BackgroundColor);
+	MyAudioMeter->SetMeterBackgroundColor(MeterBackgroundColor);
+	MyAudioMeter->SetMeterValueColor(MeterValueColor);
+	MyAudioMeter->SetMeterPeakColor(MeterPeakColor);
+	MyAudioMeter->SetMeterClippingColor(MeterClippingColor);
+	MyAudioMeter->SetMeterScaleColor(MeterScaleColor);
+	MyAudioMeter->SetMeterScaleLabelColor(MeterScaleLabelColor);
+
+	TAttribute<TArray<FMeterChannelInfo>> MeterChannelInfoBinding = PROPERTY_BINDING(TArray<FMeterChannelInfo>, MeterChannelInfo);
+	MyAudioMeter->SetMeterChannelInfo(MeterChannelInfoBinding);
+}
+
+void UAudioMeter2::ReleaseSlateResources(bool bReleaseChildren)
+{
+	Super::ReleaseSlateResources(bReleaseChildren);
+
+	MyAudioMeter.Reset();
+}
+
+TArray<FMeterChannelInfo> UAudioMeter2::GetMeterChannelInfo() const
+{
+	if (MyAudioMeter.IsValid())
+	{
+		return MyAudioMeter->GetMeterChannelInfo();
+	}
+	return TArray<FMeterChannelInfo>();
+}
+
+void UAudioMeter2::SetMeterChannelInfo(const TArray<FMeterChannelInfo>& InMeterChannelInfo)
+{
+	if (MyAudioMeter.IsValid())
+	{
+		MyAudioMeter->SetMeterChannelInfo(InMeterChannelInfo);
+	}
+}
+
+void UAudioMeter2::SetBackgroundColor(FLinearColor InValue)
+{
+	BackgroundColor = InValue;
+	if (MyAudioMeter.IsValid())
+	{
+		MyAudioMeter->SetBackgroundColor(InValue);
+	}
+}
+
+void UAudioMeter2::SetMeterBackgroundColor(FLinearColor InValue)
+{
+	MeterBackgroundColor = InValue;
+	if (MyAudioMeter.IsValid())
+	{
+		MyAudioMeter->SetMeterBackgroundColor(InValue);
+	}
+}
+
+void UAudioMeter2::SetMeterValueColor(FLinearColor InValue)
+{
+	MeterValueColor = InValue;
+	if (MyAudioMeter.IsValid())
+	{
+		MyAudioMeter->SetMeterValueColor(InValue);
+	}
+}
+
+void UAudioMeter2::SetMeterPeakColor(FLinearColor InValue)
+{
+	MeterPeakColor = InValue;
+	if (MyAudioMeter.IsValid())
+	{
+		MyAudioMeter->SetMeterPeakColor(InValue);
+	}
+}
+
+void UAudioMeter2::SetMeterClippingColor(FLinearColor InValue)
+{
+	MeterClippingColor = InValue;
+	if (MyAudioMeter.IsValid())
+	{
+		MyAudioMeter->SetMeterClippingColor(InValue);
+	}
+}
+
+void UAudioMeter2::SetMeterScaleColor(FLinearColor InValue)
+{
+	MeterScaleColor = InValue;
+	if (MyAudioMeter.IsValid())
+	{
+		MyAudioMeter->SetMeterScaleColor(InValue);
+	}
+}
+
+void UAudioMeter2::SetMeterScaleLabelColor(FLinearColor InValue)
+{
+	MeterScaleLabelColor = InValue;
+	if (MyAudioMeter.IsValid())
+	{
+		MyAudioMeter->SetMeterScaleLabelColor(InValue);
+	}
+}
+
+#if WITH_EDITOR
+
+const FText UAudioMeter2::GetPaletteCategory()
+{
+	return LOCTEXT("Common", "Common");
+}
+
+#endif
+
+/////////////////////////////////////////////////////
+
+#undef LOCTEXT_NAMESPACE
\ No newline at end of file
diff --git a/Engine/Plugins/Runtime/AudioWidgets/Source/AudioWidgets/Private/AudioMeterStyle.cpp b/Engine/Plugins/Runtime/AudioWidgets/Source/AudioWidgets/Private/AudioMeterStyle.cpp
new file mode 100644
--- /dev/null
+++ b/Engine/Plugins/Runtime/AudioWidgets/Source/AudioWidgets/Private/AudioMeterStyle.cpp
@@ -0,0 +1,35 @@
+// Copyright Epic Games, Inc. All Rights Reserved.
+
+#include "AudioMeterStyle.h"
+
+
+FAudioMeterStyle::FAudioMeterStyle()
+	: MeterSize(FVector2D(100.0f, 500.0f))
+	, MeterPadding(FVector2D(5.0f, 5.0f))
+	, MeterValuePadding(2.0f)
+	, PeakValueWidth(8.0f)
+	, ValueRangeDb(FVector2D(-60, 10))
+	, bShowScale(true)
+	, bScaleSide(true)
+	, ScaleHashOffset(5.0f)
+	, ScaleHashWidth(3.0f)
+	, ScaleHashHeight(20.0f)
+	, DecibelsPerHash(5)
+{
+}
+
+void FAudioMeterStyle::GetResources(TArray< const FSlateBrush* >& OutBrushes) const
+{
+	OutBrushes.Add(&MeterValueImage);
+	OutBrushes.Add(&MeterBackgroundImage);
+	OutBrushes.Add(&MeterPeakImage);
+}
+
+const FName FAudioMeterStyle::TypeName(TEXT("FAudioMeterStyle"));
+
+const FAudioMeterStyle& FAudioMeterStyle::GetDefault()
+{
+	static FAudioMeterStyle Default;
+	return Default;
+}
+
diff --git a/Engine/Plugins/Runtime/AudioWidgets/Source/AudioWidgets/Private/SAudioMeter.cpp b/Engine/Plugins/Runtime/AudioWidgets/Source/AudioWidgets/Private/SAudioMeter.cpp
new file mode 100644
--- /dev/null
+++ b/Engine/Plugins/Runtime/AudioWidgets/Source/AudioWidgets/Private/SAudioMeter.cpp
@@ -0,0 +1,401 @@
+// Copyright Epic Games, Inc. All Rights Reserved.
+
+#include "SAudioMeter.h"
+#include "Rendering/DrawElements.h"
+#include "Framework/Application/SlateApplication.h"
+#include "DSP/Dsp.h"
+#include "Fonts/FontMeasure.h"
+
+SAudioMeter::SAudioMeter()
+{
+}
+
+void SAudioMeter::Construct(const SAudioMeter::FArguments& InDeclaration)
+{
+	check(InDeclaration._Style);
+
+	Orientation = InDeclaration._Orientation;
+
+	BackgroundColor = InDeclaration._BackgroundColor;
+	MeterBackgroundColor = InDeclaration._MeterBackgroundColor;
+	MeterValueColor = InDeclaration._MeterValueColor;
+	MeterPeakColor = InDeclaration._MeterPeakColor;
+	MeterClippingColor = InDeclaration._MeterClippingColor;
+
+	Style = InDeclaration._Style;
+
+	MeterChannelInfoAttribute = InDeclaration._MeterChannelInfo;
+}
+
+int32 SAudioMeter::OnPaint(const FPaintArgs& Args, const FGeometry& AllottedGeometry, const FSlateRect& MyCullingRect, FSlateWindowElementList& OutDrawElements, int32 LayerId, const FWidgetStyle& InWidgetStyle, bool bParentEnabled) const
+{
+	// we draw the slider like a horizontal slider regardless of the orientation, and apply a render transform to make it display correctly.
+	// However, the AllottedGeometry is computed as it will be rendered, so we have to use the "horizontal orientation" when doing drawing computations.
+	const float AllottedWidth = Orientation == Orient_Horizontal ? AllottedGeometry.GetLocalSize().X : AllottedGeometry.GetLocalSize().Y;
+	const float AllottedHeight = Orientation == Orient_Horizontal ? AllottedGeometry.GetLocalSize().Y : AllottedGeometry.GetLocalSize().X;
+
+	FGeometry MeterGeometry = AllottedGeometry;
+
+	// rotate the slider 90deg if it's vertical. The 0 side goes on the bottom, the 1 side on the top.
+	if (Orientation == Orient_Vertical)
+	{
+		// Do this by translating along -X by the width of the geometry, then rotating 90 degrees CCW (left-hand coords)
+		FSlateRenderTransform SlateRenderTransform = TransformCast<FSlateRenderTransform>(Concatenate(Inverse(FVector2D(AllottedWidth, 0)), FQuat2D(FMath::DegreesToRadians(-90.0f))));
+		// create a child geometry matching this one, but with the render transform.
+		MeterGeometry = AllottedGeometry.MakeChild(
+			FVector2D(AllottedWidth, AllottedHeight),
+			FSlateLayoutTransform(),
+			SlateRenderTransform, FVector2D::ZeroVector);
+	}
+ 
+ 	// We clamp to make sure that the slider cannot go out of the slider Length.
+ 	int32 DecibelsPerHash = Style->DecibelsPerHash;
+ 
+ 	int32 MinValueDb = FMath::Min((int32)Style->ValueRangeDb.X, (int32)Style->ValueRangeDb.Y);
+ 	int32 MaxValueDb = FMath::Max((int32)Style->ValueRangeDb.Y, (int32)Style->ValueRangeDb.Y);
+ 
+ 	// Snap the min/max values to the nearest hash 
+ 	MaxValueDb -= MaxValueDb % DecibelsPerHash;
+ 	MinValueDb -= MinValueDb % DecibelsPerHash;
+ 
+ 	float ValueDeltaDb = MaxValueDb - MinValueDb;
+
+	// Draw the background image
+	{
+		FVector2D BackgroundImageTopLeft = FVector2D(0.0f, 0.0f);
+		FVector2D BackgroundImageSize = FVector2D(AllottedWidth, AllottedHeight);
+
+		// Draw background image
+		FSlateDrawElement::MakeBox(
+			OutDrawElements,
+			LayerId,
+			MeterGeometry.ToPaintGeometry(BackgroundImageTopLeft, BackgroundImageSize),
+			&Style->MeterBackgroundImage,
+			ESlateDrawEffect::None,
+			Style->BackgroundImage.GetTint(InWidgetStyle) * BackgroundColor.Get().GetColor(InWidgetStyle) * InWidgetStyle.GetColorAndOpacityTint()
+		);
+
+	}
+ 	++LayerId;
+
+	float ScaleOffset = 0.0f;
+	if (Style->bShowScale && Style->bScaleSide)
+	{
+		 ScaleOffset = GetScaleHeight();
+	}
+
+	TArray<FMeterChannelInfo> ChannelInfos = MeterChannelInfoAttribute.Get();
+	int32 NumChannels = ChannelInfos.Num();
+
+	// Draw the meter backgrounds
+	for (int32 ChannelIndex = 0; ChannelIndex < NumChannels; ++ChannelIndex)
+	{		
+		FVector2D MeterBackgroundTopLeft = FVector2D(Style->MeterPadding.X, Style->MeterPadding.Y + ScaleOffset + ChannelIndex * (Style->MeterSize.Y + Style->MeterPadding.Y));
+		FVector2D MeterBackgroundSize = FVector2D(AllottedWidth - 2.0f * Style->MeterPadding.X, Style->MeterSize.Y);
+
+ 		FSlateDrawElement::MakeBox(
+			OutDrawElements,
+			LayerId,
+			MeterGeometry.ToPaintGeometry(MeterBackgroundTopLeft, MeterBackgroundSize),
+			&Style->MeterBackgroundImage,
+			ESlateDrawEffect::None,
+			Style->MeterBackgroundImage.GetTint(InWidgetStyle) * MeterBackgroundColor.Get().GetColor(InWidgetStyle) * InWidgetStyle.GetColorAndOpacityTint()
+		);
+	}
+
+	++LayerId;
+
+	// Draw the meter value
+	for (int32 ChannelIndex = 0; ChannelIndex < NumChannels; ++ChannelIndex)
+	{
+		float ChannelMeterValueDb = ChannelInfos[ChannelIndex].MeterValue;
+		float ChannelMeterValue = FMath::Clamp((ChannelMeterValueDb - MinValueDb) / ValueDeltaDb, 0.0f, 1.0f);
+
+		FVector2D MeterValueTopLeft = FVector2D(Style->MeterPadding.X, 
+												Style->MeterPadding.Y + ScaleOffset + Style->MeterValuePadding + ChannelIndex * (Style->MeterSize.Y + Style->MeterPadding.Y));
+		
+		FVector2D MeterValueSize = FVector2D(ChannelMeterValue * (AllottedWidth - 2.0f * Style->MeterPadding.X), 
+										     Style->MeterSize.Y - 2.0f * Style->MeterValuePadding);
+ 
+		FSlateDrawElement::MakeBox(
+			OutDrawElements,
+			LayerId,
+			MeterGeometry.ToPaintGeometry(MeterValueTopLeft, MeterValueSize),
+			&Style->MeterValueImage,
+			ESlateDrawEffect::None,
+			Style->MeterValueImage.GetTint(InWidgetStyle) * MeterValueColor.Get().GetColor(InWidgetStyle) * InWidgetStyle.GetColorAndOpacityTint()
+		);
+	}
+
+	++LayerId;
+
+	// Draw the meter peak value
+	for (int32 ChannelIndex = 0; ChannelIndex < NumChannels; ++ChannelIndex)
+	{	
+		float ChannelPeakValueDb = ChannelInfos[ChannelIndex].PeakValue;
+		float ChannelPeakValue = FMath::Clamp((ChannelPeakValueDb - MinValueDb) / ValueDeltaDb, 0.0f, 1.0f);
+
+		float PeakPixel = ChannelPeakValue * (AllottedWidth - 2.0f * Style->MeterPadding.X);
+
+		FVector2D MeterPeakValueTopLeft = FVector2D(Style->MeterPadding.X + PeakPixel - 0.5f * Style->PeakValueWidth, 
+													Style->MeterPadding.Y + ScaleOffset + Style->MeterValuePadding + ChannelIndex * (Style->MeterSize.Y + Style->MeterPadding.Y));
+		
+		FVector2D MeterPeakValueSize = FVector2D(Style->PeakValueWidth, 
+												 Style->MeterSize.Y - 2.0f * Style->MeterValuePadding);
+
+		FSlateDrawElement::MakeBox(
+			OutDrawElements,
+			LayerId,
+			MeterGeometry.ToPaintGeometry(MeterPeakValueTopLeft, MeterPeakValueSize),
+			&Style->MeterPeakImage,
+			ESlateDrawEffect::None,
+			Style->MeterPeakImage.GetTint(InWidgetStyle) * MeterPeakColor.Get().GetColor(InWidgetStyle) * InWidgetStyle.GetColorAndOpacityTint()
+		);
+	}
+	
+	++LayerId;
+
+	// Draw the scale hash
+	if (Style->bShowScale)
+	{
+
+		float ScaleHashHalfWidth = 0.5f * Style->ScaleHashWidth;
+		FVector2D HashSize = FVector2D(Style->ScaleHashWidth, Style->ScaleHashHeight);
+
+		const TSharedRef<FSlateFontMeasure> FontMeasureService = FSlateApplication::Get().GetRenderer()->GetFontMeasureService();
+
+		// Measure the min value label size so we can right-justify if needed
+		FVector2D MinValueLabelSize = FontMeasureService->Measure(FString::FromInt(MinValueDb), Style->Font);
+
+		// Get the size of the negative sign to use to offset the label text in horizontal mode
+		FVector2D NegativeSignSize;
+		if (Orientation == Orient_Horizontal)
+		{
+			NegativeSignSize = FontMeasureService->Measure(TEXT("-"), Style->Font);
+		}
+
+		int32 ValueDelta = MaxValueDb - MinValueDb;
+		int32 CurrentHashValue = MaxValueDb;
+		while (CurrentHashValue >= MinValueDb)
+		{
+			// Get the fractional value for hash mark
+			const float CurrentHashMeterValuePercent = FMath::Clamp(((float)CurrentHashValue - (float)MinValueDb) / ValueDelta, 0.0f, 1.0f);
+
+			// Draw hash
+			float HashPixelCenter = CurrentHashMeterValuePercent * (AllottedWidth - 2.0f * Style->MeterPadding.X);
+
+			FVector2D HashTopLeft;
+			HashTopLeft.X = Style->MeterPadding.X + HashPixelCenter - ScaleHashHalfWidth;
+			if (Style->bScaleSide)
+			{
+				HashTopLeft.Y = Style->MeterPadding.Y + ScaleOffset - Style->ScaleHashOffset - Style->ScaleHashHeight;
+			}
+			else
+			{
+				HashTopLeft.Y = (Style->MeterSize.Y + Style->MeterPadding.Y ) * NumChannels + Style->ScaleHashOffset;
+			}
+
+			FSlateDrawElement::MakeBox(
+				OutDrawElements,
+				LayerId,
+				MeterGeometry.ToPaintGeometry(HashTopLeft, HashSize),
+				&Style->MeterPeakImage,
+				ESlateDrawEffect::None,
+				MeterScaleColor.Get().GetColor(InWidgetStyle) * InWidgetStyle.GetColorAndOpacityTint()
+			);
+		
+			FString LabelString = FString::FromInt(CurrentHashValue);
+
+			bool bIsNegative = (CurrentHashValue < 0);
+
+			FVector2D LabelSize; 
+			if (Orientation == Orient_Horizontal && bIsNegative)
+			{
+				// We want to center the text on just the positive portion of the number
+				FString LabelStringPositive = FString::FromInt(FMath::Abs(CurrentHashValue));
+				LabelSize = FontMeasureService->Measure(LabelStringPositive, Style->Font);
+			}
+			else
+			{
+				LabelSize = FontMeasureService->Measure(LabelString, Style->Font);
+			}
+			
+			FText LabelText = FText::FromString(LabelString);
+
+			FGeometry TextGeometry = MeterGeometry;
+			if (Orientation == Orient_Vertical)
+			{
+				FVector2D LabelTopLeft;
+				
+				LabelTopLeft.X = Style->MeterPadding.X + HashPixelCenter + 0.5f * LabelSize.Y;
+
+				if (Style->bScaleSide)
+				{
+					LabelTopLeft.Y = Style->MeterPadding.Y - 2.0f + (MinValueLabelSize.X - LabelSize.X);
+				}
+				else
+				{
+					LabelTopLeft.Y = (Style->MeterSize.Y + Style->MeterPadding.Y) * NumChannels + Style->ScaleHashOffset + Style->ScaleHashHeight + 2.0f;
+				}
+				
+				// Undo the rotation for vertical right before we do the rendering of the scale value 
+				FSlateRenderTransform RotationTransform = FSlateRenderTransform(FQuat2D(FMath::DegreesToRadians(90.0f)));
+
+				FSlateLayoutTransform ChildLayoutTransform(1.0f, TransformPoint(1.0f, LabelTopLeft));
+
+				TextGeometry = MeterGeometry.MakeChild(
+					LabelSize,
+					ChildLayoutTransform,
+					RotationTransform,
+					FVector2D(0.0f, 0.0f)
+				);
+
+				FSlateDrawElement::MakeText(
+					OutDrawElements,
+					LayerId,
+					TextGeometry.ToOffsetPaintGeometry(FVector2D(0.0f, 0.0f)),
+					LabelText,
+					Style->Font,
+					ESlateDrawEffect::None,
+					MeterScaleLabelColor.Get().GetColor(InWidgetStyle)* InWidgetStyle.GetColorAndOpacityTint()
+				);
+			}
+			else
+			{
+				FVector2D LabelTopLeft;
+
+				// Center it based off the positive value width, then subtract from there the size of the negative sign
+				LabelTopLeft.X = Style->MeterPadding.X + HashPixelCenter - 0.5f * LabelSize.X;
+				if (bIsNegative)
+				{
+					LabelTopLeft.X -= NegativeSignSize.X;
+				}
+
+				if (Style->bScaleSide)
+				{
+					LabelTopLeft.Y = Style->MeterPadding.Y + ScaleOffset - LabelSize.Y - Style->ScaleHashOffset - Style->ScaleHashHeight;
+				}
+				else
+				{
+					LabelTopLeft.Y = (Style->MeterSize.Y + Style->MeterPadding.Y) * NumChannels + Style->ScaleHashOffset + Style->ScaleHashHeight + 2.0f;
+				}
+
+				// Draw text label
+				FSlateDrawElement::MakeText(
+					OutDrawElements,
+					LayerId,
+					TextGeometry.ToOffsetPaintGeometry(LabelTopLeft),
+					LabelText,
+					Style->Font,
+					ESlateDrawEffect::None,
+					MeterScaleLabelColor.Get().GetColor(InWidgetStyle) * InWidgetStyle.GetColorAndOpacityTint()
+				);
+
+			}
+
+			CurrentHashValue -= Style->DecibelsPerHash;
+		}
+	}
+
+	return LayerId;
+}
+
+float SAudioMeter::GetScaleHeight() const
+{
+	const TSharedRef<FSlateFontMeasure> FontMeasureService = FSlateApplication::Get().GetRenderer()->GetFontMeasureService();
+	FVector2D LabelSize = FontMeasureService->Measure(FString::FromInt(-60), Style->Font);
+
+	float ScaleHeight = Style->ScaleHashHeight + Style->ScaleHashOffset;
+	if (Orientation == Orient_Vertical)
+	{
+		return ScaleHeight + LabelSize.X;
+	}
+	return ScaleHeight + LabelSize.Y;
+}
+
+FVector2D SAudioMeter::ComputeDesiredSize(float) const
+{
+	static const FVector2D SAudioMeterDesiredSize(50.0f, 50.0f);
+
+	if (Style == nullptr)
+	{
+		return SAudioMeterDesiredSize;
+	}
+
+	TArray<FMeterChannelInfo> ChannelInfo = MeterChannelInfoAttribute.Get();
+	int32 NumChannels = FMath::Max(ChannelInfo.Num(), 1);
+
+	FVector2D Size = FVector2D(Style->MeterSize.X, (Style->MeterSize.Y + Style->MeterPadding.Y) * NumChannels);
+
+	// Add the end padding
+	Size.X += Style->MeterPadding.X;
+	Size.Y += Style->MeterPadding.Y;
+
+	// Add the width for the scale if it's been set to show
+	if (Style->bShowScale)
+	{
+		Size.Y += GetScaleHeight();
+	}
+
+	// The thickness is going to be the thickness plus the scale image width
+	if (Orientation == Orient_Vertical)
+	{
+		return FVector2D(Size.Y, Size.X);
+	}
+
+	return Size; 
+}
+
+void SAudioMeter::SetMeterChannelInfo(const TAttribute<TArray<FMeterChannelInfo>>& InMeterChannelInfo)
+{
+	SetAttribute(MeterChannelInfoAttribute, InMeterChannelInfo, EInvalidateWidgetReason::Paint);
+}
+
+TArray<FMeterChannelInfo> SAudioMeter::GetMeterChannelInfo() const
+{
+	return MeterChannelInfoAttribute.Get();
+}
+
+void SAudioMeter::SetOrientation(EOrientation InOrientation)
+{
+	if (Orientation != InOrientation)
+	{
+		Orientation = InOrientation;
+		Invalidate(EInvalidateWidgetReason::Layout);
+	}
+}
+
+void SAudioMeter::SetBackgroundColor(FSlateColor InBackgroundColor)
+{
+	SetAttribute(BackgroundColor, TAttribute<FSlateColor>(InBackgroundColor), EInvalidateWidgetReason::Paint);
+}
+
+void SAudioMeter::SetMeterBackgroundColor(FSlateColor InMeterBackgroundColor)
+{
+	SetAttribute(MeterBackgroundColor, TAttribute<FSlateColor>(InMeterBackgroundColor), EInvalidateWidgetReason::Paint);
+}
+
+void SAudioMeter::SetMeterValueColor(FSlateColor InMeterValueColor)
+{
+	SetAttribute(MeterValueColor, TAttribute<FSlateColor>(InMeterValueColor), EInvalidateWidgetReason::Paint);
+}
+
+void SAudioMeter::SetMeterPeakColor(FSlateColor InMeterPeakColor)
+{
+	SetAttribute(MeterPeakColor, TAttribute<FSlateColor>(InMeterPeakColor), EInvalidateWidgetReason::Paint);
+}
+
+void SAudioMeter::SetMeterClippingColor(FSlateColor InMeterClippingColor)
+{
+	SetAttribute(MeterClippingColor, TAttribute<FSlateColor>(InMeterClippingColor), EInvalidateWidgetReason::Paint);
+}
+
+void SAudioMeter::SetMeterScaleColor(FSlateColor InMeterScaleColor)
+{
+	SetAttribute(MeterScaleColor, TAttribute<FSlateColor>(InMeterScaleColor), EInvalidateWidgetReason::Paint);
+}
+
+void SAudioMeter::SetMeterScaleLabelColor(FSlateColor InMeterScaleLabelColor)
+{
+	SetAttribute(MeterScaleLabelColor, TAttribute<FSlateColor>(InMeterScaleLabelColor), EInvalidateWidgetReason::Paint);
+}
diff --git a/Engine/Plugins/Runtime/AudioWidgets/Source/AudioWidgets/Public/AudioMeter.h b/Engine/Plugins/Runtime/AudioWidgets/Source/AudioWidgets/Public/AudioMeter.h
new file mode 100644
--- /dev/null
+++ b/Engine/Plugins/Runtime/AudioWidgets/Source/AudioWidgets/Public/AudioMeter.h
@@ -0,0 +1,140 @@
+// Copyright Epic Games, Inc. All Rights Reserved.
+
+#pragma once
+
+#include "CoreMinimal.h"
+#include "UObject/ObjectMacros.h"
+#include "Styling/SlateTypes.h"
+#include "Widgets/SWidget.h"
+#include "Components/Widget.h"
+#include "AudioMeterStyle.h"
+#include "AudioMeterTypes.h"
+#include "AudioMeter.generated.h"
+
+class SAudioMeter;
+
+/**
+ * An audio meter widget.
+ *
+ * Supports displaying a slower moving peak-hold value as well as the current meter value.
+ *
+ * A clipping value is also displayed which shows a customizable color to indicate clipping.
+ *
+ * Internal values are stored and interacted with as linear volume values.
+ * 
+ */
+UCLASS()
+class AUDIOWIDGETS_API UAudioMeter2: public UWidget
+{
+	GENERATED_UCLASS_BODY()
+
+	public:
+
+	DECLARE_DYNAMIC_DELEGATE_RetVal(TArray<FMeterChannelInfo>, FGetMeterChannelInfo);
+
+	/** The current meter value to display. */
+	UPROPERTY(EditAnywhere, Category = MeterValues)
+	TArray<FMeterChannelInfo> MeterChannelInfo;
+
+	/** A bindable delegate to allow logic to drive the value of the meter */
+	UPROPERTY()
+	FGetMeterChannelInfo MeterChannelInfoDelegate;
+
+public:
+	
+	/** The audio meter style */
+	UPROPERTY(EditAnywhere, BlueprintReadWrite, Category="Style", meta=( DisplayName="Style" ))
+	FAudioMeterStyle WidgetStyle;
+
+	/** The slider's orientation. */
+	UPROPERTY(EditAnywhere, BlueprintReadOnly, Category=Appearance)
+	TEnumAsByte<EOrientation> Orientation;
+
+	/** The color to draw the background. */
+	UPROPERTY(EditAnywhere, BlueprintReadOnly, Category = Appearance)
+	FLinearColor BackgroundColor;
+
+	/** The color to draw the meter background. */
+	UPROPERTY(EditAnywhere, BlueprintReadOnly, Category=Appearance)
+	FLinearColor MeterBackgroundColor;
+
+	/** The color to draw the meter value. */
+	UPROPERTY(EditAnywhere, BlueprintReadOnly, Category = Appearance)
+	FLinearColor MeterValueColor;
+
+	/** The color to draw the meter peak value. */
+	UPROPERTY(EditAnywhere, BlueprintReadOnly, Category = Appearance)
+	FLinearColor MeterPeakColor;
+
+	/** The color to draw the meter clipping value. */
+	UPROPERTY(EditAnywhere, BlueprintReadOnly, Category = Appearance)
+	FLinearColor MeterClippingColor;
+
+	/** The color to draw the meter scale hashes. */
+	UPROPERTY(EditAnywhere, BlueprintReadOnly, Category = Appearance)
+	FLinearColor MeterScaleColor;
+
+	/** The color to draw the meter scale label. */
+	UPROPERTY(EditAnywhere, BlueprintReadOnly, Category = Appearance)
+	FLinearColor MeterScaleLabelColor;
+
+public:
+
+ 	/** Gets the current linear value of the meter. */
+ 	UFUNCTION(BlueprintCallable, Category="Behavior")
+	TArray<FMeterChannelInfo> GetMeterChannelInfo() const;
+ 
+ 	/** Sets the current meter values. */
+ 	UFUNCTION(BlueprintCallable, Category="Behavior")
+ 	void SetMeterChannelInfo(const TArray<FMeterChannelInfo>& InMeterChannelInfo);
+
+	/** Sets the background color */
+	UFUNCTION(BlueprintCallable, Category = "Appearance")
+	void SetBackgroundColor(FLinearColor InValue);
+
+	/** Sets the meter background color */
+	UFUNCTION(BlueprintCallable, Category="Appearance")
+	void SetMeterBackgroundColor(FLinearColor InValue);
+
+	/** Sets the meter value color */
+	UFUNCTION(BlueprintCallable, Category = "Appearance")
+	void SetMeterValueColor(FLinearColor InValue);
+
+	/** Sets the meter peak color */
+	UFUNCTION(BlueprintCallable, Category = "Appearance")
+	void SetMeterPeakColor(FLinearColor InValue);
+
+	/** Sets the meter clipping color */
+	UFUNCTION(BlueprintCallable, Category = "Appearance")
+	void SetMeterClippingColor(FLinearColor InValue);
+
+	/** Sets the meter scale color */
+	UFUNCTION(BlueprintCallable, Category = "Appearance")
+	void SetMeterScaleColor(FLinearColor InValue);
+
+	/** Sets the meter scale color */
+	UFUNCTION(BlueprintCallable, Category = "Appearance")
+	void SetMeterScaleLabelColor(FLinearColor InValue);
+
+	// UWidget interface
+	virtual void SynchronizeProperties() override;
+	// End of UWidget interface
+
+	// UVisual interface
+	virtual void ReleaseSlateResources(bool bReleaseChildren) override;
+	// End of UVisual interface
+
+#if WITH_EDITOR
+	virtual const FText GetPaletteCategory() override;
+#endif
+
+protected:
+	/** Native Slate Widget */
+	TSharedPtr<SAudioMeter> MyAudioMeter;
+
+	// UWidget interface
+	virtual TSharedRef<SWidget> RebuildWidget() override;
+	// End of UWidget interface
+
+	PROPERTY_BINDING_IMPLEMENTATION(TArray<FMeterChannelInfo>, MeterChannelInfo);
+};
\ No newline at end of file
diff --git a/Engine/Plugins/Runtime/AudioWidgets/Source/AudioWidgets/Public/AudioMeterStyle.h b/Engine/Plugins/Runtime/AudioWidgets/Source/AudioWidgets/Public/AudioMeterStyle.h
new file mode 100644
--- /dev/null
+++ b/Engine/Plugins/Runtime/AudioWidgets/Source/AudioWidgets/Public/AudioMeterStyle.h
@@ -0,0 +1,141 @@
+// Copyright Epic Games, Inc. All Rights Reserved.
+
+#pragma once
+
+#include "CoreMinimal.h"
+#include "UObject/ObjectMacros.h"
+#include "Widgets/SWidget.h"
+#include "Components/Widget.h"
+#include "Styling/SlateWidgetStyle.h"
+#include "AudioMeterStyle.generated.h"
+
+/**
+ * Represents the appearance of an SAudioMeter
+ */
+USTRUCT(BlueprintType)
+struct AUDIOWIDGETS_API FAudioMeterStyle : public FSlateWidgetStyle
+{
+	GENERATED_USTRUCT_BODY()
+
+	FAudioMeterStyle();
+
+	virtual ~FAudioMeterStyle() {}
+
+	virtual void GetResources( TArray< const FSlateBrush* >& OutBrushes ) const override;
+
+	static const FName TypeName;
+	virtual const FName GetTypeName() const override { return TypeName; };
+
+	static const FAudioMeterStyle& GetDefault();
+
+	/** Image to use to represent the meter value. */
+	UPROPERTY(EditAnywhere, BlueprintReadWrite, Category=Appearance)
+	FSlateBrush MeterValueImage;
+	FAudioMeterStyle& SetMeterValueImage(const FSlateBrush& InMeterValueImage){ MeterValueImage = InMeterValueImage; return *this; }
+
+	/** Image to use to represent the background. */
+	UPROPERTY(EditAnywhere, BlueprintReadWrite, Category = Appearance)
+	FSlateBrush BackgroundImage;
+	FAudioMeterStyle& SetBackgroundImage(const FSlateBrush& InBackgroundImage) { BackgroundImage = InBackgroundImage; return *this; }
+
+	/** Image to use to represent the meter background. */
+	UPROPERTY(EditAnywhere, BlueprintReadWrite, Category = Appearance)
+	FSlateBrush MeterBackgroundImage;
+	FAudioMeterStyle& SetMeterBackgroundImage(const FSlateBrush& InMeterBackgroundImage) { MeterBackgroundImage = InMeterBackgroundImage; return *this; }
+
+	/** Image to use to draw behind the meter value. */
+	UPROPERTY(EditAnywhere, BlueprintReadWrite, Category = Appearance)
+	FSlateBrush MeterValueBackgroundImage;
+	FAudioMeterStyle& SetMeterValueBackgroundImage(const FSlateBrush& InMeterValueBackgroundImage) { MeterValueBackgroundImage = InMeterValueBackgroundImage; return *this; }
+
+	/** Image to use to represent the meter peak. */
+	UPROPERTY(EditAnywhere, BlueprintReadWrite, Category = Appearance)
+	FSlateBrush MeterPeakImage;
+	FAudioMeterStyle& SetMeterPeakImage(const FSlateBrush& InMeterPeakImage) { MeterPeakImage = InMeterPeakImage; return *this; }
+
+	// How thick to draw the meter
+	UPROPERTY(EditAnywhere, BlueprintReadWrite, Category=Appearance)
+	FVector2D MeterSize;
+	FAudioMeterStyle& SetMeterSize(const FVector2D& InMeterSize) { MeterSize = InMeterSize; return *this; }
+
+	// How much padding to add around the meter
+	UPROPERTY(EditAnywhere, BlueprintReadWrite, Category = Appearance)
+	FVector2D MeterPadding;
+	FAudioMeterStyle& SetMeterPadding(const FVector2D& InMeterPadding) { MeterPadding = InMeterPadding; return *this; }
+
+	// How much padding to add around the meter
+	UPROPERTY(EditAnywhere, BlueprintReadWrite, Category = Appearance)
+	float MeterValuePadding;
+	FAudioMeterStyle& SetMeterValuePadding(float InMeterValuePadding) { MeterValuePadding = InMeterValuePadding; return *this; }
+
+	// How wide to draw the peak value indicator
+	UPROPERTY(EditAnywhere, BlueprintReadWrite, Category = Appearance)
+	float PeakValueWidth;
+	FAudioMeterStyle& SetPeakValueWidth(float InPeakValueWidth) { PeakValueWidth = InPeakValueWidth; return *this; }
+
+	// The minimum and maximum value to display in dB (values are clamped in this range)
+	UPROPERTY(EditAnywhere, BlueprintReadWrite, Category = Appearance)
+	FVector2D ValueRangeDb;
+	FAudioMeterStyle& SetValueRangeDb(const FVector2D& InValueRangeDb) { ValueRangeDb = InValueRangeDb; return *this; }
+
+	// Whether or not to show the decibel scale alongside the meter
+	UPROPERTY(EditAnywhere, BlueprintReadWrite, Category = Appearance)
+	bool bShowScale;
+	FAudioMeterStyle& SetShowScale(bool bInShowScale) { bShowScale = bInShowScale; return *this; }
+
+	// Which side to show the scale. If vertical, true means left side, false means right side. If horizontal, true means above, false means below.
+	UPROPERTY(EditAnywhere, BlueprintReadWrite, Category = Appearance)
+	bool bScaleSide;
+	FAudioMeterStyle& SetScaleSide(bool bInScaleSide) { bScaleSide = bInScaleSide; return *this; }
+
+	// Offset for the hashes
+	UPROPERTY(EditAnywhere, BlueprintReadWrite, Category = Appearance)
+	float ScaleHashOffset;
+	FAudioMeterStyle& SetScaleHashOffset(float InScaleHashOffset) { ScaleHashOffset = InScaleHashOffset; return *this; }
+
+	// The width of each hash mark
+	UPROPERTY(EditAnywhere, BlueprintReadWrite, Category = Appearance)
+	float ScaleHashWidth;
+	FAudioMeterStyle& SetScaleHashWidth(float InScaleHashWidth) { ScaleHashWidth = InScaleHashWidth; return *this; }
+
+	// The height of each hash mark
+	UPROPERTY(EditAnywhere, BlueprintReadWrite, Category = Appearance)
+	float ScaleHashHeight;
+	FAudioMeterStyle& SetScaleHashHeight(float InScaleHashHeight) { ScaleHashHeight = InScaleHashHeight; return *this; }
+
+	// How wide to draw the decibel scale, if it's enabled
+	UPROPERTY(EditAnywhere, BlueprintReadWrite, Category = Appearance, meta = (UIMin = "3", ClampMin="3", UIMax = "10"))
+	int32 DecibelsPerHash;
+	FAudioMeterStyle& SetDecibelsPerHash(float InDecibelsPerHash) { DecibelsPerHash = InDecibelsPerHash; return *this; }
+
+	/** Font family and size to be used when displaying the meter scale. */
+	UPROPERTY(EditAnywhere, BlueprintReadWrite, Category = Appearance)
+	FSlateFontInfo Font;
+	FAudioMeterStyle& SetFont(const FSlateFontInfo& InFont) { Font = InFont; return *this; }
+	FAudioMeterStyle& SetFont(TSharedPtr<const FCompositeFont> InCompositeFont, const int32 InSize, const FName& InTypefaceFontName = NAME_None) { Font = FSlateFontInfo(InCompositeFont, InSize, InTypefaceFontName); return *this; }
+	FAudioMeterStyle& SetFont(const UObject* InFontObject, const int32 InSize, const FName& InTypefaceFontName = NAME_None) { Font = FSlateFontInfo(InFontObject, InSize, InTypefaceFontName); return *this; }
+	FAudioMeterStyle& SetFont(const FName& InFontName, uint16 InSize) { Font = FSlateFontInfo(InFontName, InSize); return *this; }
+	FAudioMeterStyle& SetFont(const FString& InFontName, uint16 InSize) { Font = FSlateFontInfo(*InFontName, InSize); return *this; }
+	FAudioMeterStyle& SetFont(const WIDECHAR* InFontName, uint16 InSize) { Font = FSlateFontInfo(InFontName, InSize); return *this; }
+	FAudioMeterStyle& SetFont(const ANSICHAR* InFontName, uint16 InSize) { Font = FSlateFontInfo(InFontName, InSize); return *this; }
+	FAudioMeterStyle& SetFontName(const FName& InFontName) { Font = FSlateFontInfo(InFontName, Font.Size); return *this; }
+	FAudioMeterStyle& SetFontName(const FString& InFontName) { Font = FSlateFontInfo(InFontName, Font.Size); return *this; }
+	FAudioMeterStyle& SetFontName(const WIDECHAR* InFontName) { Font = FSlateFontInfo(InFontName, Font.Size); return *this; }
+	FAudioMeterStyle& SetFontName(const ANSICHAR* InFontName) { Font = FSlateFontInfo(InFontName, Font.Size); return *this; }
+	FAudioMeterStyle& SetFontSize(uint16 InSize) { Font.Size = InSize; return *this; }
+	FAudioMeterStyle& SetTypefaceFontName(const FName& InTypefaceFontName) { Font.TypefaceFontName = InTypefaceFontName; return *this; }
+
+	/**
+	 * Unlinks all colors in this style.
+	 * @see FSlateColor::Unlink
+	 */
+	void UnlinkColors()
+	{
+		MeterValueImage.UnlinkColors();
+		MeterBackgroundImage.UnlinkColors();
+//		MeterScaleImage.UnlinkColors();
+		MeterPeakImage.UnlinkColors();
+	}
+};
+
+
diff --git a/Engine/Plugins/Runtime/AudioWidgets/Source/AudioWidgets/Public/AudioMeterTypes.h b/Engine/Plugins/Runtime/AudioWidgets/Source/AudioWidgets/Public/AudioMeterTypes.h
new file mode 100644
--- /dev/null
+++ b/Engine/Plugins/Runtime/AudioWidgets/Source/AudioWidgets/Public/AudioMeterTypes.h
@@ -0,0 +1,28 @@
+// Copyright Epic Games, Inc. All Rights Reserved.
+
+#pragma once
+
+#include "CoreMinimal.h"
+#include "AudioMeterTypes.generated.h"
+
+USTRUCT(BlueprintType)
+struct FMeterChannelInfo
+{
+	GENERATED_BODY()
+
+	UPROPERTY(EditAnywhere, BlueprintReadWrite, Category = AudioMeter)
+	float MeterValue = 0.0f;
+
+	UPROPERTY(EditAnywhere, BlueprintReadWrite, Category = AudioMeter)
+	float PeakValue = 0.0f;
+
+	UPROPERTY(EditAnywhere, BlueprintReadWrite, Category = AudioMeter)
+	float ClippingValue = 0.0f;
+};
+
+inline bool operator==(const FMeterChannelInfo& lhs, const FMeterChannelInfo& rhs)
+{
+	return FMath::IsNearlyEqual(lhs.MeterValue, rhs.MeterValue) &&
+		FMath::IsNearlyEqual(lhs.PeakValue, rhs.PeakValue) &&
+		FMath::IsNearlyEqual(lhs.ClippingValue, rhs.ClippingValue);
+}
diff --git a/Engine/Plugins/Runtime/AudioWidgets/Source/AudioWidgets/Public/SAudioMeter.h b/Engine/Plugins/Runtime/AudioWidgets/Source/AudioWidgets/Public/SAudioMeter.h
new file mode 100644
--- /dev/null
+++ b/Engine/Plugins/Runtime/AudioWidgets/Source/AudioWidgets/Public/SAudioMeter.h
@@ -0,0 +1,127 @@
+// Copyright Epic Games, Inc. All Rights Reserved.
+
+#pragma once
+
+#include "CoreMinimal.h"
+#include "Misc/Attribute.h"
+#include "Input/Reply.h"
+#include "Styling/SlateWidgetStyleAsset.h"
+#include "Widgets/DeclarativeSyntaxSupport.h"
+#include "Styling/SlateColor.h"
+#include "Styling/SlateTypes.h"
+#include "Styling/CoreStyle.h"
+#include "Framework/SlateDelegates.h"
+#include "Widgets/SLeafWidget.h"
+#include "AudioMeterStyle.h"
+#include "AudioMeterTypes.h"
+
+class FPaintArgs;
+class FSlateWindowElementList;
+
+/**
+ * A Slate slider control is a linear scale and draggable handle.
+ */
+class AUDIOWIDGETS_API SAudioMeter
+	: public SLeafWidget
+{
+public:
+
+	SLATE_BEGIN_ARGS(SAudioMeter)
+		: _Orientation(EOrientation::Orient_Horizontal)
+		, _BackgroundColor(FLinearColor::Black)
+		, _MeterBackgroundColor(FLinearColor::Gray)
+		, _MeterValueColor(FLinearColor::Green)
+		, _MeterPeakColor(FLinearColor::Blue)
+		, _MeterScaleColor(FLinearColor::White)
+		, _MeterScaleLabelColor(FLinearColor::Gray)
+		, _Style(&FCoreStyle::Get().GetWidgetStyle<FAudioMeterStyle>("AudioMeter"))
+	{
+	}
+
+	/** Whether the slidable area should be indented to fit the handle. */
+	SLATE_ATTRIBUTE(bool, IndentHandle)
+
+		/** The slider's orientation. */
+		SLATE_ARGUMENT(EOrientation, Orientation)
+
+		/** The color to draw the background in. */
+		SLATE_ATTRIBUTE(FSlateColor, BackgroundColor)
+
+		/** The color to draw the meter background in. */
+		SLATE_ATTRIBUTE(FSlateColor, MeterBackgroundColor)
+
+		/** The color to draw the meter value in. */
+		SLATE_ATTRIBUTE(FSlateColor, MeterValueColor)
+
+		/** The color to draw the meter peak. */
+		SLATE_ATTRIBUTE(FSlateColor, MeterPeakColor)
+
+		/** The color to draw the clipping value in. */
+		SLATE_ATTRIBUTE(FSlateColor, MeterClippingColor)
+
+		/** The color to draw the scale in. */
+		SLATE_ATTRIBUTE(FSlateColor, MeterScaleColor)
+
+		/** The color to draw the scale in. */
+		SLATE_ATTRIBUTE(FSlateColor, MeterScaleLabelColor)
+
+		/** The style used to draw the slider. */
+		SLATE_STYLE_ARGUMENT(FAudioMeterStyle, Style)
+
+		/** A value representing the audio meter value. */
+		SLATE_ATTRIBUTE(TArray<FMeterChannelInfo>, MeterChannelInfo)
+
+		SLATE_END_ARGS()
+
+		SAudioMeter();
+
+	/**
+	 * Construct the widget.
+	 *
+	 * @param InDeclaration A declaration from which to construct the widget.
+	 */
+	void Construct(const SAudioMeter::FArguments& InDeclaration);
+
+	void SetMeterChannelInfo(const TAttribute<TArray<FMeterChannelInfo>>& InMeterChannelInfo);
+	TArray<FMeterChannelInfo> GetMeterChannelInfo() const;
+
+	/** Set the Orientation attribute */
+	void SetOrientation(EOrientation InOrientation);
+
+	void SetBackgroundColor(FSlateColor InBackgroundColor);
+	void SetMeterBackgroundColor(FSlateColor InMeterBackgroundColor);
+	void SetMeterValueColor(FSlateColor InMeterValueColor);
+	void SetMeterPeakColor(FSlateColor InMeterPeakColor);
+	void SetMeterClippingColor(FSlateColor InMeterPeakColor);
+	void SetMeterScaleColor(FSlateColor InMeterScaleColor);
+	void SetMeterScaleLabelColor(FSlateColor InMeterScaleLabelColor);
+
+public:
+
+	// SWidget overrides
+
+	virtual int32 OnPaint(const FPaintArgs& Args, const FGeometry& AllottedGeometry, const FSlateRect& MyCullingRect, FSlateWindowElementList& OutDrawElements, int32 LayerId, const FWidgetStyle& InWidgetStyle, bool bParentEnabled) const override;
+	virtual FVector2D ComputeDesiredSize(float) const override;
+
+protected:
+	
+	// Returns the scale height based off font size and hash height
+	float GetScaleHeight() const;
+	
+	// Holds the style passed to the widget upon construction.
+	const FAudioMeterStyle* Style;
+
+	// Holds the slider's orientation.
+	EOrientation Orientation;
+
+	// Various colors
+	TAttribute<FSlateColor> BackgroundColor;
+	TAttribute<FSlateColor> MeterBackgroundColor;
+	TAttribute<FSlateColor> MeterValueColor;
+	TAttribute<FSlateColor> MeterPeakColor;
+	TAttribute<FSlateColor> MeterClippingColor;
+	TAttribute<FSlateColor> MeterScaleColor;
+	TAttribute<FSlateColor> MeterScaleLabelColor;
+
+	TAttribute<TArray<FMeterChannelInfo>> MeterChannelInfoAttribute;
+};
