Renaming UAudioMeter to the correct name

diff --git a/Engine/Plugins/Runtime/AudioWidgets/Source/AudioWidgets/Private/AudioMeter.cpp b/Engine/Plugins/Runtime/AudioWidgets/Source/AudioWidgets/Private/AudioMeter.cpp
--- a/Engine/Plugins/Runtime/AudioWidgets/Source/AudioWidgets/Private/AudioMeter.cpp
+++ b/Engine/Plugins/Runtime/AudioWidgets/Source/AudioWidgets/Private/AudioMeter.cpp
@@ -12,7 +12,7 @@
 
 static FAudioMeterStyle* DefaultAudioMeterStyle = nullptr;
 
-UAudioMeter2::UAudioMeter2(const FObjectInitializer& ObjectInitializer)
+UAudioMeter::UAudioMeter(const FObjectInitializer& ObjectInitializer)
 	: Super(ObjectInitializer)
 {
 	Orientation = EOrientation::Orient_Vertical;
@@ -42,7 +42,7 @@ UAudioMeter2::UAudioMeter2(const FObjectInitializer& ObjectInitializer)
 #endif
 }
 
-TSharedRef<SWidget> UAudioMeter2::RebuildWidget()
+TSharedRef<SWidget> UAudioMeter::RebuildWidget()
 {
 	MyAudioMeter = SNew(SAudioMeter)
 		.Style(&WidgetStyle);
@@ -50,7 +50,7 @@ TSharedRef<SWidget> UAudioMeter2::RebuildWidget()
 	return MyAudioMeter.ToSharedRef();
 }
 
-void UAudioMeter2::SynchronizeProperties()
+void UAudioMeter::SynchronizeProperties()
 {
 	Super::SynchronizeProperties();
 
@@ -68,14 +68,14 @@ void UAudioMeter2::SynchronizeProperties()
 	MyAudioMeter->SetMeterChannelInfo(MeterChannelInfoBinding);
 }
 
-void UAudioMeter2::ReleaseSlateResources(bool bReleaseChildren)
+void UAudioMeter::ReleaseSlateResources(bool bReleaseChildren)
 {
 	Super::ReleaseSlateResources(bReleaseChildren);
 
 	MyAudioMeter.Reset();
 }
 
-TArray<FMeterChannelInfo> UAudioMeter2::GetMeterChannelInfo() const
+TArray<FMeterChannelInfo> UAudioMeter::GetMeterChannelInfo() const
 {
 	if (MyAudioMeter.IsValid())
 	{
@@ -84,7 +84,7 @@ TArray<FMeterChannelInfo> UAudioMeter2::GetMeterChannelInfo() const
 	return TArray<FMeterChannelInfo>();
 }
 
-void UAudioMeter2::SetMeterChannelInfo(const TArray<FMeterChannelInfo>& InMeterChannelInfo)
+void UAudioMeter::SetMeterChannelInfo(const TArray<FMeterChannelInfo>& InMeterChannelInfo)
 {
 	if (MyAudioMeter.IsValid())
 	{
@@ -92,7 +92,7 @@ void UAudioMeter2::SetMeterChannelInfo(const TArray<FMeterChannelInfo>& InMeterC
 	}
 }
 
-void UAudioMeter2::SetBackgroundColor(FLinearColor InValue)
+void UAudioMeter::SetBackgroundColor(FLinearColor InValue)
 {
 	BackgroundColor = InValue;
 	if (MyAudioMeter.IsValid())
@@ -101,7 +101,7 @@ void UAudioMeter2::SetBackgroundColor(FLinearColor InValue)
 	}
 }
 
-void UAudioMeter2::SetMeterBackgroundColor(FLinearColor InValue)
+void UAudioMeter::SetMeterBackgroundColor(FLinearColor InValue)
 {
 	MeterBackgroundColor = InValue;
 	if (MyAudioMeter.IsValid())
@@ -110,7 +110,7 @@ void UAudioMeter2::SetMeterBackgroundColor(FLinearColor InValue)
 	}
 }
 
-void UAudioMeter2::SetMeterValueColor(FLinearColor InValue)
+void UAudioMeter::SetMeterValueColor(FLinearColor InValue)
 {
 	MeterValueColor = InValue;
 	if (MyAudioMeter.IsValid())
@@ -119,7 +119,7 @@ void UAudioMeter2::SetMeterValueColor(FLinearColor InValue)
 	}
 }
 
-void UAudioMeter2::SetMeterPeakColor(FLinearColor InValue)
+void UAudioMeter::SetMeterPeakColor(FLinearColor InValue)
 {
 	MeterPeakColor = InValue;
 	if (MyAudioMeter.IsValid())
@@ -128,7 +128,7 @@ void UAudioMeter2::SetMeterPeakColor(FLinearColor InValue)
 	}
 }
 
-void UAudioMeter2::SetMeterClippingColor(FLinearColor InValue)
+void UAudioMeter::SetMeterClippingColor(FLinearColor InValue)
 {
 	MeterClippingColor = InValue;
 	if (MyAudioMeter.IsValid())
@@ -137,7 +137,7 @@ void UAudioMeter2::SetMeterClippingColor(FLinearColor InValue)
 	}
 }
 
-void UAudioMeter2::SetMeterScaleColor(FLinearColor InValue)
+void UAudioMeter::SetMeterScaleColor(FLinearColor InValue)
 {
 	MeterScaleColor = InValue;
 	if (MyAudioMeter.IsValid())
@@ -146,7 +146,7 @@ void UAudioMeter2::SetMeterScaleColor(FLinearColor InValue)
 	}
 }
 
-void UAudioMeter2::SetMeterScaleLabelColor(FLinearColor InValue)
+void UAudioMeter::SetMeterScaleLabelColor(FLinearColor InValue)
 {
 	MeterScaleLabelColor = InValue;
 	if (MyAudioMeter.IsValid())
@@ -157,7 +157,7 @@ void UAudioMeter2::SetMeterScaleLabelColor(FLinearColor InValue)
 
 #if WITH_EDITOR
 
-const FText UAudioMeter2::GetPaletteCategory()
+const FText UAudioMeter::GetPaletteCategory()
 {
 	return LOCTEXT("Common", "Common");
 }
diff --git a/Engine/Plugins/Runtime/AudioWidgets/Source/AudioWidgets/Public/AudioMeter.h b/Engine/Plugins/Runtime/AudioWidgets/Source/AudioWidgets/Public/AudioMeter.h
--- a/Engine/Plugins/Runtime/AudioWidgets/Source/AudioWidgets/Public/AudioMeter.h
+++ b/Engine/Plugins/Runtime/AudioWidgets/Source/AudioWidgets/Public/AudioMeter.h
@@ -24,7 +24,7 @@ class SAudioMeter;
  * 
  */
 UCLASS()
-class AUDIOWIDGETS_API UAudioMeter2: public UWidget
+class AUDIOWIDGETS_API UAudioMeter: public UWidget
 {
 	GENERATED_UCLASS_BODY()
 
