Making default appearance of audio meter look decent.

WIP, will set these to application defaults or some other method in the future.

diff --git a/Engine/Plugins/Runtime/AudioWidgets/Source/AudioWidgets/Private/AudioMeter.cpp b/Engine/Plugins/Runtime/AudioWidgets/Source/AudioWidgets/Private/AudioMeter.cpp
--- a/Engine/Plugins/Runtime/AudioWidgets/Source/AudioWidgets/Private/AudioMeter.cpp
+++ b/Engine/Plugins/Runtime/AudioWidgets/Source/AudioWidgets/Private/AudioMeter.cpp
@@ -17,13 +17,19 @@ UAudioMeter::UAudioMeter(const FObjectInitializer& ObjectInitializer)
 {
 	Orientation = EOrientation::Orient_Vertical;
 
-	BackgroundColor = FLinearColor::Black;
-	MeterBackgroundColor = FLinearColor::Gray;
-	MeterValueColor = FLinearColor::Green;
-	MeterPeakColor = FLinearColor::Blue;
-	MeterClippingColor = FLinearColor::Red;
-	MeterScaleColor = FLinearColor::Gray;
-	MeterScaleLabelColor = FLinearColor::White;
+	BackgroundColor = FLinearColor(0.0075f, 0.0075f, 0.0075, 1.0f);
+	MeterBackgroundColor = FLinearColor(0.031f, 0.031f, 0.031f, 1.0f);
+	MeterValueColor = FLinearColor(0.025719f, 0.208333f, 0.069907f, 1.0f);
+	MeterPeakColor = FLinearColor(0.24349f, 0.708333f, 0.357002f, 1.0f);
+	MeterClippingColor = FLinearColor(1.0f, 0.0f, 0.112334f, 1.0f);
+	MeterScaleColor = FLinearColor(0.017642f, 0.017642f, 0.017642f, 1.0f);
+	MeterScaleLabelColor = FLinearColor(0.442708f, 0.442708f, 0.442708f, 1.0f);
+
+	// Add a single channel as a default just so it can be seen when somebody makes one
+	FMeterChannelInfo DefaultInfo;
+	DefaultInfo.MeterValue = -6.0f;
+	DefaultInfo.PeakValue = -3.0f;
+	MeterChannelInfo.Add(DefaultInfo);
 
 	if (DefaultAudioMeterStyle == nullptr)
 	{
diff --git a/Engine/Plugins/Runtime/AudioWidgets/Source/AudioWidgets/Private/AudioMeterStyle.cpp b/Engine/Plugins/Runtime/AudioWidgets/Source/AudioWidgets/Private/AudioMeterStyle.cpp
--- a/Engine/Plugins/Runtime/AudioWidgets/Source/AudioWidgets/Private/AudioMeterStyle.cpp
+++ b/Engine/Plugins/Runtime/AudioWidgets/Source/AudioWidgets/Private/AudioMeterStyle.cpp
@@ -1,20 +1,21 @@
 // Copyright Epic Games, Inc. All Rights Reserved.
 
 #include "AudioMeterStyle.h"
-
+#include "Styling/StyleDefaults.h"
 
 FAudioMeterStyle::FAudioMeterStyle()
-	: MeterSize(FVector2D(100.0f, 500.0f))
-	, MeterPadding(FVector2D(5.0f, 5.0f))
-	, MeterValuePadding(2.0f)
-	, PeakValueWidth(8.0f)
+	: MeterSize(FVector2D(250.0f, 25.0f))
+	, MeterPadding(FVector2D(10.0f, 5.0f))
+	, MeterValuePadding(3.0f)
+	, PeakValueWidth(2.0f)
 	, ValueRangeDb(FVector2D(-60, 10))
 	, bShowScale(true)
 	, bScaleSide(true)
 	, ScaleHashOffset(5.0f)
-	, ScaleHashWidth(3.0f)
-	, ScaleHashHeight(20.0f)
+	, ScaleHashWidth(1.0f)
+	, ScaleHashHeight(10.0f)
 	, DecibelsPerHash(5)
+	, Font(FStyleDefaults::GetFontInfo(5))
 {
 }
 
diff --git a/Engine/Plugins/Runtime/AudioWidgets/Source/AudioWidgets/Public/AudioMeterStyle.h b/Engine/Plugins/Runtime/AudioWidgets/Source/AudioWidgets/Public/AudioMeterStyle.h
--- a/Engine/Plugins/Runtime/AudioWidgets/Source/AudioWidgets/Public/AudioMeterStyle.h
+++ b/Engine/Plugins/Runtime/AudioWidgets/Source/AudioWidgets/Public/AudioMeterStyle.h
@@ -63,7 +63,7 @@ struct AUDIOWIDGETS_API FAudioMeterStyle : public FSlateWidgetStyle
 	FVector2D MeterPadding;
 	FAudioMeterStyle& SetMeterPadding(const FVector2D& InMeterPadding) { MeterPadding = InMeterPadding; return *this; }
 
-	// How much padding to add around the meter
+	// How much padding to add around the meter value
 	UPROPERTY(EditAnywhere, BlueprintReadWrite, Category = Appearance)
 	float MeterValuePadding;
 	FAudioMeterStyle& SetMeterValuePadding(float InMeterValuePadding) { MeterValuePadding = InMeterValuePadding; return *this; }
