Fixing warning

diff --git a/Engine/Plugins/Runtime/AudioWidgets/Source/AudioWidgets/Private/AudioMeter.cpp b/Engine/Plugins/Runtime/AudioWidgets/Source/AudioWidgets/Private/AudioMeter.cpp
--- a/Engine/Plugins/Runtime/AudioWidgets/Source/AudioWidgets/Private/AudioMeter.cpp
+++ b/Engine/Plugins/Runtime/AudioWidgets/Source/AudioWidgets/Private/AudioMeter.cpp
@@ -10,8 +10,6 @@
 /////////////////////////////////////////////////////
 // UAudioMeter
 
-static FAudioMeterStyle* DefaultAudioMeterStyle = nullptr;
-
 UAudioMeter::UAudioMeter(const FObjectInitializer& ObjectInitializer)
 	: Super(ObjectInitializer)
 {
@@ -31,16 +29,7 @@ UAudioMeter::UAudioMeter(const FObjectInitializer& ObjectInitializer)
 	DefaultInfo.PeakValue = -3.0f;
 	MeterChannelInfo.Add(DefaultInfo);
 
-	if (DefaultAudioMeterStyle == nullptr)
-	{
-		// HACK: THIS SHOULD NOT COME FROM CORESTYLE AND SHOULD INSTEAD BE DEFINED BY ENGINE TEXTURES/PROJECT SETTINGS
-		DefaultAudioMeterStyle = new FAudioMeterStyle(FCoreStyle::Get().GetWidgetStyle<FAudioMeterStyle>("AudioMeter"));
-
-		// Unlink UMG default colors from the editor settings colors.
-		DefaultAudioMeterStyle->UnlinkColors();
-	}
-
-	WidgetStyle = *DefaultAudioMeterStyle;
+	WidgetStyle = FAudioMeterStyle::GetDefault();
 
 #if WITH_EDITORONLY_DATA
 	AccessibleBehavior = ESlateAccessibleBehavior::NotAccessible;
