Slight fix up to audio widgets plugin crashing on load and fixing the audio meter category.

diff --git a/Engine/Plugins/Runtime/AudioWidgets/Source/AudioWidgets/Private/AudioMeter.cpp b/Engine/Plugins/Runtime/AudioWidgets/Source/AudioWidgets/Private/AudioMeter.cpp
--- a/Engine/Plugins/Runtime/AudioWidgets/Source/AudioWidgets/Private/AudioMeter.cpp
+++ b/Engine/Plugins/Runtime/AudioWidgets/Source/AudioWidgets/Private/AudioMeter.cpp
@@ -154,7 +154,7 @@ void UAudioMeter::SetMeterScaleLabelColor(FLinearColor InValue)
 
 const FText UAudioMeter::GetPaletteCategory()
 {
-	return LOCTEXT("Common", "Common");
+	return LOCTEXT("Audio", "Audio");
 }
 
 #endif
diff --git a/Engine/Plugins/Runtime/AudioWidgets/Source/AudioWidgets/Private/AudioWidgets.cpp b/Engine/Plugins/Runtime/AudioWidgets/Source/AudioWidgets/Private/AudioWidgets.cpp
--- a/Engine/Plugins/Runtime/AudioWidgets/Source/AudioWidgets/Private/AudioWidgets.cpp
+++ b/Engine/Plugins/Runtime/AudioWidgets/Source/AudioWidgets/Private/AudioWidgets.cpp
@@ -8,7 +8,7 @@
 
 void FAudioWidgetsModule::StartupModule()
 {
-	FModuleManager::Get().LoadModuleChecked<FAdvancedWidgetsModule>("AdvancedWidgets");
+	//FModuleManager::Get().LoadModuleChecked<FAdvancedWidgetsModule>("AdvancedWidgets");
 }
 
 void FAudioWidgetsModule::ShutdownModule()
