Add back module load to ensure audio utility widgets are available to UMG

diff --git a/Engine/Plugins/Runtime/AudioWidgets/Source/AudioWidgets/Private/AudioWidgets.cpp b/Engine/Plugins/Runtime/AudioWidgets/Source/AudioWidgets/Private/AudioWidgets.cpp
--- a/Engine/Plugins/Runtime/AudioWidgets/Source/AudioWidgets/Private/AudioWidgets.cpp
+++ b/Engine/Plugins/Runtime/AudioWidgets/Source/AudioWidgets/Private/AudioWidgets.cpp
@@ -8,7 +8,9 @@
 
 void FAudioWidgetsModule::StartupModule()
 {
-	//FModuleManager::Get().LoadModuleChecked<FAdvancedWidgetsModule>("AdvancedWidgets");
+	// Required to load so the AudioWidget plugin content can reference widgets
+	// defined in AdvancedWidgets (ex. RadialSlider for UMG-defined knobs)
+	FModuleManager::Get().LoadModuleChecked<FAdvancedWidgetsModule>("AdvancedWidgets");
 }
 
 void FAudioWidgetsModule::ShutdownModule()
