Audio Widgets - nativized Audio Sliders
- slider with a editable text label for audio specific applications
- versions for slider with customizable linear to output curve, volume slider, and frequency slider
- UWidget versions of Slate objects
- various customizable properties like colors and text label shown on hover

diff --git a/Engine/Plugins/Runtime/AudioWidgets/Source/AudioWidgets/AudioWidgets.Build.cs b/Engine/Plugins/Runtime/AudioWidgets/Source/AudioWidgets/AudioWidgets.Build.cs
--- a/Engine/Plugins/Runtime/AudioWidgets/Source/AudioWidgets/AudioWidgets.Build.cs
+++ b/Engine/Plugins/Runtime/AudioWidgets/Source/AudioWidgets/AudioWidgets.Build.cs
@@ -25,7 +25,7 @@ public class AudioWidgets : ModuleRules
 				"SlateCore",
 				"InputCore",
 				"AdvancedWidgets",
-				"SignalProcessing"
+				"SignalProcessing",
 			}
 		);
 	}
diff --git a/Engine/Plugins/Runtime/AudioWidgets/Source/AudioWidgets/Private/AudioSlider.cpp b/Engine/Plugins/Runtime/AudioWidgets/Source/AudioWidgets/Private/AudioSlider.cpp
new file mode 100644
--- /dev/null
+++ b/Engine/Plugins/Runtime/AudioWidgets/Source/AudioWidgets/Private/AudioSlider.cpp
@@ -0,0 +1,194 @@
+// Copyright Epic Games, Inc. All Rights Reserved.
+
+#include "AudioSlider.h"
+#include "SAudioSlider.h"
+#include "Widgets/DeclarativeSyntaxSupport.h"
+
+#define LOCTEXT_NAMESPACE "AUDIO_UMG"
+
+// UAudioSliderBase
+
+UAudioSliderBase::UAudioSliderBase(const FObjectInitializer& ObjectInitializer)
+	: Super(ObjectInitializer)
+{
+	Value = 0.0f;
+	AlwaysShowLabel = false;
+	LabelBackgroundColor = FLinearColor(0.01033f, 0.01033f, 0.01033f);
+	SliderBackgroundColor = FLinearColor(0.01033f, 0.01033f, 0.01033f);
+	SliderBarColor = FLinearColor::Black;
+	SliderThumbColor = FLinearColor::Gray;
+	WidgetBackgroundColor = FLinearColor::Transparent;
+
+#if WITH_EDITORONLY_DATA
+	AccessibleBehavior = ESlateAccessibleBehavior::NotAccessible;
+	bCanChildrenBeAccessible = false;
+#endif
+}
+
+void UAudioSliderBase::SynchronizeProperties()
+{
+	Super::SynchronizeProperties();
+	
+	MyAudioSlider->SetAlwaysShowLabel(AlwaysShowLabel);
+	MyAudioSlider->SetLabelBackgroundColor(LabelBackgroundColor);
+	MyAudioSlider->SetSliderBackgroundColor(SliderBackgroundColor);
+	MyAudioSlider->SetSliderBarColor(SliderBarColor);
+	MyAudioSlider->SetSliderThumbColor(SliderThumbColor);
+	MyAudioSlider->SetWidgetBackgroundColor(WidgetBackgroundColor);
+}
+
+void UAudioSliderBase::ReleaseSlateResources(bool bReleaseChildren)
+{
+	Super::ReleaseSlateResources(bReleaseChildren);
+
+	MyAudioSlider.Reset();
+}
+
+void UAudioSliderBase::HandleOnValueChanged(float InValue)
+{
+	OnValueChanged.Broadcast(InValue);
+}
+
+void UAudioSliderBase::SetTextReadOnly(const bool bIsReadOnly)
+{
+	MyAudioSlider->SetTextReadOnly(bIsReadOnly);
+}
+
+void UAudioSliderBase::SetAlwaysShowLabel(const bool bSetAlwaysShowLabel)
+{
+	AlwaysShowLabel = bSetAlwaysShowLabel;
+	MyAudioSlider->SetAlwaysShowLabel(bSetAlwaysShowLabel);
+}
+
+void UAudioSliderBase::SetUnitsText(const FText Units)
+{
+	MyAudioSlider->SetUnitsText(Units);
+}
+
+float UAudioSliderBase::GetOutputValue(const float LinValue)
+{
+	return MyAudioSlider->GetOutputValue(LinValue);
+}
+
+float UAudioSliderBase::GetLinValue(const float OutputValue)
+{
+	return MyAudioSlider->GetLinValue(OutputValue);
+}
+
+void UAudioSliderBase::SetLabelBackgroundColor(FLinearColor InValue)
+{
+	SliderBackgroundColor = InValue;
+	if (MyAudioSlider.IsValid())
+	{
+		MyAudioSlider->SetLabelBackgroundColor(InValue);
+	}
+}
+
+void UAudioSliderBase::SetSliderBackgroundColor(FLinearColor InValue)
+{
+	SliderBackgroundColor = InValue;
+	if (MyAudioSlider.IsValid())
+	{
+		MyAudioSlider->SetSliderBackgroundColor(InValue);
+	}
+}
+
+void UAudioSliderBase::SetSliderBarColor(FLinearColor InValue)
+{
+	SliderBarColor = InValue;
+	if (MyAudioSlider.IsValid())
+	{
+		MyAudioSlider->SetSliderBarColor(InValue);
+	}
+}
+
+void UAudioSliderBase::SetSliderThumbColor(FLinearColor InValue)
+{
+	SliderThumbColor = InValue;
+	if (MyAudioSlider.IsValid())
+	{
+		MyAudioSlider->SetSliderThumbColor(InValue);
+	}
+}
+
+void UAudioSliderBase::SetWidgetBackgroundColor(FLinearColor InValue)
+{
+	WidgetBackgroundColor = InValue;
+	if (MyAudioSlider.IsValid())
+	{
+		MyAudioSlider->SetWidgetBackgroundColor(InValue);
+	}
+}
+
+TSharedRef<SWidget> UAudioSliderBase::RebuildWidget()
+{
+	MyAudioSlider = SNew(SAudioSliderBase)
+		.OnValueChanged(BIND_UOBJECT_DELEGATE(FOnFloatValueChanged, HandleOnValueChanged));
+
+	return MyAudioSlider.ToSharedRef();
+}
+
+#if WITH_EDITOR
+
+const FText UAudioSliderBase::GetPaletteCategory()
+{
+	return LOCTEXT("Audio", "Audio");
+}
+
+#endif
+
+// UAudioSlider
+UAudioSlider::UAudioSlider(const FObjectInitializer& ObjectInitializer)
+	: Super(ObjectInitializer)
+{
+}
+
+void UAudioSlider::SynchronizeProperties()
+{
+	UAudioSliderBase::SynchronizeProperties();
+
+	StaticCastSharedPtr<SAudioSlider>(MyAudioSlider)->SetLinToOutputCurve(LinToOutputCurve);
+	StaticCastSharedPtr<SAudioSlider>(MyAudioSlider)->SetOutputToLinCurve(OutputToLinCurve);
+}
+
+TSharedRef<SWidget> UAudioSlider::RebuildWidget()
+{
+	MyAudioSlider = SNew(SAudioSlider)
+		.OnValueChanged(BIND_UOBJECT_DELEGATE(FOnFloatValueChanged, HandleOnValueChanged));
+
+	return MyAudioSlider.ToSharedRef();
+}
+
+// UAudioVolumeSlider
+UAudioVolumeSlider::UAudioVolumeSlider(const FObjectInitializer& ObjectInitializer)
+	: Super(ObjectInitializer)
+{
+}
+
+TSharedRef<SWidget> UAudioVolumeSlider::RebuildWidget()
+{
+	MyAudioSlider = SNew(SAudioVolumeSlider)
+		.OnValueChanged(BIND_UOBJECT_DELEGATE(FOnFloatValueChanged, HandleOnValueChanged));
+	LinToOutputCurve = StaticCastSharedPtr<SAudioSlider>(MyAudioSlider)->GetLinToOutputCurve();
+	OutputToLinCurve = StaticCastSharedPtr<SAudioSlider>(MyAudioSlider)->GetOutputToLinCurve();
+
+	return MyAudioSlider.ToSharedRef();
+}
+
+// UAudioFrequencySlider
+UAudioFrequencySlider::UAudioFrequencySlider(const FObjectInitializer& ObjectInitializer)
+	: Super(ObjectInitializer)
+	, OutputRange(FVector2D(20.0f, 20000.0f))
+{
+}
+
+TSharedRef<SWidget> UAudioFrequencySlider::RebuildWidget()
+{
+	MyAudioSlider = SNew(SAudioFrequencySlider)
+		.OnValueChanged(BIND_UOBJECT_DELEGATE(FOnFloatValueChanged, HandleOnValueChanged));
+	StaticCastSharedPtr<SAudioFrequencySlider>(MyAudioSlider)->SetOutputRange(OutputRange);
+
+	return MyAudioSlider.ToSharedRef();
+}
+
+#undef LOCTEXT_NAMESPACE
\ No newline at end of file
diff --git a/Engine/Plugins/Runtime/AudioWidgets/Source/AudioWidgets/Private/AudioSliderStyle.cpp b/Engine/Plugins/Runtime/AudioWidgets/Source/AudioWidgets/Private/AudioSliderStyle.cpp
new file mode 100644
--- /dev/null
+++ b/Engine/Plugins/Runtime/AudioWidgets/Source/AudioWidgets/Private/AudioSliderStyle.cpp
@@ -0,0 +1,61 @@
+// Copyright Epic Games, Inc. All Rights Reserved.
+
+#include "AudioSliderStyle.h"
+
+FAudioSliderStyle::FAudioSliderStyle()
+	: FSlateStyleSet("AudioSliderStyle")
+{
+	SetParentStyleName("CoreStyle");
+	SetContentRoot(FPaths::EnginePluginsDir() / TEXT("Runtime/AudioWidgets/Content/AudioSlider"));
+
+	const float LabelWidth = 56.0f;
+	const float LabelHeight = 22.0f;
+	const FVector2D LabelBackgroundSize = FVector2D(LabelWidth, LabelHeight);
+	const FVector2D ThumbImageSize = FVector2D(22.0f, 22.0f);
+
+	const float SliderBarWidth = 10.0f;
+	const float SliderBarHeight = 432.0f;
+	const FVector2D SliderBarCapSize = FVector2D(SliderBarWidth, SliderBarWidth / 2.0f);
+	const FVector2D SliderBarRectangleSize = FVector2D(SliderBarWidth, SliderBarHeight - SliderBarWidth);
+
+	const float SliderBackgroundWidth = 28.0f;
+	const float SliderBackgroundHeight = 450.0f;
+	const FVector2D SliderBackgroundCapSize = FVector2D(SliderBackgroundWidth, SliderBackgroundWidth / 2.0f);
+	const FVector2D SliderBackgroundRectangleSize = FVector2D(SliderBackgroundWidth, SliderBackgroundHeight - SliderBackgroundWidth);
+	const FVector2D SliderBackgroundSize = FVector2D(SliderBackgroundWidth, SliderBackgroundHeight);
+	
+	const float LabelVerticalPadding = 5.0f;
+	const FVector2D DesiredWidgetSize = FVector2D(LabelWidth, LabelHeight + LabelVerticalPadding + SliderBackgroundHeight);
+	
+	Set("AudioSlider.LabelVerticalPadding", LabelVerticalPadding);
+	Set("AudioSlider.DesiredWidgetSize", DesiredWidgetSize);
+
+	Set("AudioSlider.LabelBackground", new IMAGE_BRUSH_SVG(TEXT("LabelBackground"), LabelBackgroundSize));
+	Set("AudioSlider.SliderBarCap", new IMAGE_BRUSH_SVG(TEXT("SliderBarCap"), SliderBarCapSize));
+	Set("AudioSlider.SliderBarRectangle", new IMAGE_BRUSH_SVG(TEXT("SliderBarRectangle"), SliderBarRectangleSize));
+	Set("AudioSlider.SliderBackgroundCap", new IMAGE_BRUSH_SVG(TEXT("SliderBackgroundCap"), SliderBackgroundCapSize));
+	Set("AudioSlider.SliderBackgroundRectangle", new IMAGE_BRUSH_SVG(TEXT("SliderBackgroundRectangle"), SliderBackgroundRectangleSize));
+	Set("AudioSlider.WidgetBackground", new IMAGE_BRUSH_SVG(TEXT("WidgetBackground"), DesiredWidgetSize));
+
+	Set("AudioSlider.DefaultBackgroundColor", FLinearColor(0.01033f, 0.01033f, 0.01033f));
+	Set("AudioSlider.DefaultBarColor", FLinearColor::Black);
+	Set("AudioSlider.DefaultThumbColor", FLinearColor::Gray);
+	Set("AudioSlider.DefaultWidgetBackgroundColor", FLinearColor::Transparent);
+
+	FSlateBrush NoImageThumb = FSlateBrush();
+	NoImageThumb.SetImageSize(ThumbImageSize);
+	NoImageThumb.DrawAs = ESlateBrushDrawType::NoDrawType;
+	FSlateBrush NoImageSlider = FSlateBrush();
+	NoImageSlider.SetImageSize(SliderBackgroundSize);
+	NoImageSlider.DrawAs = ESlateBrushDrawType::NoDrawType;
+
+	Set("AudioSlider.Slider", FSliderStyle()
+		.SetNormalBarImage(NoImageSlider)
+		.SetHoveredBarImage(NoImageSlider)
+		.SetDisabledBarImage(NoImageSlider)
+		.SetNormalThumbImage(IMAGE_BRUSH_SVG(TEXT("SliderThumb"), ThumbImageSize))
+		.SetHoveredThumbImage(IMAGE_BRUSH_SVG(TEXT("SliderThumb"), ThumbImageSize))
+		.SetDisabledThumbImage(NoImageThumb)
+	);
+	FSlateStyleRegistry::RegisterSlateStyle(*this);
+}
diff --git a/Engine/Plugins/Runtime/AudioWidgets/Source/AudioWidgets/Private/AudioWidgets.cpp b/Engine/Plugins/Runtime/AudioWidgets/Source/AudioWidgets/Private/AudioWidgets.cpp
--- a/Engine/Plugins/Runtime/AudioWidgets/Source/AudioWidgets/Private/AudioWidgets.cpp
+++ b/Engine/Plugins/Runtime/AudioWidgets/Source/AudioWidgets/Private/AudioWidgets.cpp
@@ -1,7 +1,6 @@
 // Copyright Epic Games, Inc. All Rights Reserved.
 
 #include "AudioWidgets.h"
-
 #include "AdvancedWidgetsModule.h"
 
 #define LOCTEXT_NAMESPACE "FAudioWidgetsModule"
@@ -11,6 +10,7 @@ void FAudioWidgetsModule::StartupModule()
 	// Required to load so the AudioWidget plugin content can reference widgets
 	// defined in AdvancedWidgets (ex. RadialSlider for UMG-defined knobs)
 	FModuleManager::Get().LoadModuleChecked<FAdvancedWidgetsModule>("AdvancedWidgets");
+	AudioSliderStyleSet = MakeShared<FAudioSliderStyle>();
 }
 
 void FAudioWidgetsModule::ShutdownModule()
diff --git a/Engine/Plugins/Runtime/AudioWidgets/Source/AudioWidgets/Private/SAudioSlider.cpp b/Engine/Plugins/Runtime/AudioWidgets/Source/AudioWidgets/Private/SAudioSlider.cpp
new file mode 100644
--- /dev/null
+++ b/Engine/Plugins/Runtime/AudioWidgets/Source/AudioWidgets/Private/SAudioSlider.cpp
@@ -0,0 +1,364 @@
+// Copyright Epic Games, Inc. All Rights Reserved.
+
+#include "SAudioSlider.h"
+#include "Brushes/SlateImageBrush.h"
+#include "DSP/Dsp.h"
+#include "Fonts/FontMeasure.h"
+#include "Framework/Application/SlateApplication.h"
+#include "Rendering/DrawElements.h"
+#include "Styling/SlateStyle.h"
+#include "Styling/SlateStyleMacros.h"
+#include "Styling/SlateStyleRegistry.h"
+
+// SAudioSliderBase
+SAudioSliderBase::SAudioSliderBase()
+{
+}
+
+void SAudioSliderBase::Construct(const SAudioSliderBase::FArguments& InDeclaration)
+{
+	OnValueChanged = InDeclaration._OnValueChanged;
+	ValueAttribute = InDeclaration._Value;
+	AlwaysShowLabel = InDeclaration._AlwaysShowLabel;
+	LabelBackgroundColor = InDeclaration._LabelBackgroundColor;
+	SliderBackgroundColor = InDeclaration._SliderBackgroundColor;
+	SliderBarColor = InDeclaration._SliderBarColor;
+	SliderThumbColor = InDeclaration._SliderThumbColor;
+	WidgetBackgroundColor = InDeclaration._WidgetBackgroundColor;
+
+	const ISlateStyle* AudioSliderStyle = FSlateStyleRegistry::FindSlateStyle("AudioSliderStyle");
+	ensure(AudioSliderStyle);
+
+	ChildSlot
+	[
+		SNew(SOverlay)
+		// Widget background image
+		+ SOverlay::Slot()
+		.HAlign(HAlign_Center)
+		.VAlign(VAlign_Fill)
+		[
+			SAssignNew(WidgetBackgroundImage, SImage)
+			.Image(AudioSliderStyle->GetBrush("AudioSlider.WidgetBackground"))
+			.ColorAndOpacity(WidgetBackgroundColor)
+		]
+		+ SOverlay::Slot()
+		.HAlign(HAlign_Center)
+		.VAlign(VAlign_Fill)
+		[
+			SNew(SVerticalBox)
+			 // Text Label
+			+ SVerticalBox::Slot()
+			.Padding(0.0f, AudioSliderStyle->GetFloat("AudioSlider.LabelVerticalPadding"))
+			.AutoHeight()
+			.HAlign(HAlign_Center)
+			.VAlign(VAlign_Center)
+			[
+				SNew(SConstraintCanvas)
+				.Visibility_Lambda([this]() 
+				{
+					return (AlwaysShowLabel.Get() || this->IsHovered()) ? EVisibility::Visible : EVisibility::Hidden;
+				})
+				+ SConstraintCanvas::Slot()
+				.Anchors(FAnchors(0.0f, 0.5f, 1.0f, 0.5f))
+				.Alignment(FVector2D(0.0f, 0.5f))
+				.Offset(FMargin(0.0f, 0.0f, 0.0f, 24.0f))
+				[
+					SAssignNew(LabelBackgroundImage, SImage)
+					.Image(AudioSliderStyle->GetBrush("AudioSlider.LabelBackground"))
+					.ColorAndOpacity(LabelBackgroundColor)
+				]
+				+ SConstraintCanvas::Slot()
+				.Anchors(FAnchors(0.0f, 0.5f, 1.0f, 0.5f))
+				.Alignment(FVector2D(0.0f, 0.5f))
+				.Offset(FMargin(0.0f, 0.0f, 0.0f, 32.0f))
+				.AutoSize(true)
+				//.ZOrder(3)
+				[
+					SNew(SHorizontalBox)
+					+ SHorizontalBox::Slot()
+					//.AutoWidth()
+					.FillWidth(0.5)
+					.Padding(4.0f, 0.0f)
+					.HAlign(HAlign_Right)
+					.VAlign(VAlign_Center)
+					[
+						SAssignNew(ValueText, SEditableText)
+						.Text(FText::AsNumber(GetOutputValue(ValueAttribute.Get())))
+						.Justification(ETextJustify::Right)
+						.OnTextCommitted_Lambda([this](const FText& Text, ETextCommit::Type CommitType)
+						{
+							const float OutputValue = FCString::Atof(*Text.ToString());
+							const float LinValue = GetLinValue(OutputValue);
+							// Update this widget's value 
+							ValueAttribute.Set(LinValue);
+							// Update the slider's value attribute
+							Slider->SetValue(ValueAttribute);
+							// Call delegate if bound
+							OnValueChanged.ExecuteIfBound(LinValue);
+						})
+					]
+					+ SHorizontalBox::Slot()
+					.AutoWidth()
+					.Padding(2.0f, 0.0f, 6.0f, 0.0f)
+					.HAlign(HAlign_Left)
+					.VAlign(VAlign_Center)
+					[
+						SAssignNew(UnitsText, SEditableText)
+						.Text(FText::FromString("units"))
+					]
+				]
+			]
+			// Slider
+			+ SVerticalBox::Slot()
+			.HAlign(HAlign_Center)
+			.VAlign(VAlign_Fill)
+			[
+				SNew(SOverlay)
+				// Slider background image
+				+ SOverlay::Slot()
+				.HAlign(HAlign_Center)
+				.VAlign(VAlign_Fill)
+				[
+					SNew(SVerticalBox)
+					+SVerticalBox::Slot()
+					.HAlign(HAlign_Center)
+					.VAlign(VAlign_Fill)
+					.AutoHeight()
+					[
+						SAssignNew(SliderBackgroundTopCapImage, SImage)
+						.Image(AudioSliderStyle->GetBrush("AudioSlider.SliderBackgroundCap"))
+						.ColorAndOpacity(SliderBackgroundColor)
+					]
+					+ SVerticalBox::Slot()
+					.HAlign(HAlign_Center)
+					.VAlign(VAlign_Fill)
+					.FillHeight(1.0f)
+					[
+						SAssignNew(SliderBackgroundRectangleImage, SImage)
+						.Image(AudioSliderStyle->GetBrush("AudioSlider.SliderBackgroundRectangle"))
+						.ColorAndOpacity(SliderBackgroundColor)
+					]
+					+ SVerticalBox::Slot()
+					.HAlign(HAlign_Center)
+					.VAlign(VAlign_Fill)
+					.AutoHeight()
+					[
+						SAssignNew(SliderBackgroundBottomCapImage, SImage)
+						.Image(AudioSliderStyle->GetBrush("AudioSlider.SliderBackgroundCap"))
+						.ColorAndOpacity(SliderBackgroundColor)
+						.RenderTransform(FSlateRenderTransform(FQuat2D(FMath::DegreesToRadians(180.0f))))
+						.RenderTransformPivot(FVector2D(0.5f, 0.5f))
+					]
+				]
+				// Slider bar image
+				+ SOverlay::Slot()
+				.HAlign(HAlign_Center)
+				.VAlign(VAlign_Fill)
+				.Padding(0.0f, 10.0f)
+				[
+					SNew(SVerticalBox)
+					+ SVerticalBox::Slot()
+					.HAlign(HAlign_Center)
+					.VAlign(VAlign_Fill)
+					.AutoHeight()
+					[
+						SAssignNew(SliderBarTopCapImage, SImage)
+						.Image(AudioSliderStyle->GetBrush("AudioSlider.SliderBarCap"))
+						.ColorAndOpacity(SliderBarColor)
+					]
+					+ SVerticalBox::Slot()
+					.HAlign(HAlign_Center)
+					.VAlign(VAlign_Fill)
+					.FillHeight(1.0f)
+					[
+						SAssignNew(SliderBarRectangleImage, SImage)
+						.Image(AudioSliderStyle->GetBrush("AudioSlider.SliderBarRectangle"))
+						.ColorAndOpacity(SliderBarColor)
+					]
+					+ SVerticalBox::Slot()
+					.HAlign(HAlign_Center)
+					.VAlign(VAlign_Fill)
+					.AutoHeight()
+					[
+						SAssignNew(SliderBarBottomCapImage, SImage)
+						.Image(AudioSliderStyle->GetBrush("AudioSlider.SliderBarCap"))
+						.ColorAndOpacity(SliderBarColor)
+						.RenderTransform(FSlateRenderTransform(FQuat2D(FMath::DegreesToRadians(180.0f))))
+						.RenderTransformPivot(FVector2D(0.5f, 0.5f))
+					]
+				]
+				+ SOverlay::Slot()
+				// Actual SSlider
+				.HAlign(HAlign_Center)
+				.VAlign(VAlign_Fill)
+				[
+					SAssignNew(Slider, SSlider)
+					.Orientation(Orient_Vertical)
+					.Value(InDeclaration._Value)
+					.Style(&AudioSliderStyle->GetWidgetStyle<FSliderStyle>("AudioSlider.Slider"))
+					.OnValueChanged_Lambda([this](float Value)
+					{
+						// Update this widget's value 
+						ValueAttribute.Set(Value);
+						// Call delegate if bound
+						OnValueChanged.ExecuteIfBound(Value);
+						// Convert lin to output, update this widget's text
+						const float OutputValue = GetOutputValue(Value);
+						ValueText->SetText(FText::AsNumber(OutputValue));
+					})
+				]
+			]
+		]
+	];
+}
+
+FVector2D SAudioSliderBase::ComputeDesiredSize(float) const
+{
+	const ISlateStyle* AudioSliderStyle = FSlateStyleRegistry::FindSlateStyle("AudioSliderStyle");
+	ensure(AudioSliderStyle);
+
+	return AudioSliderStyle->GetVector("AudioSlider.DesiredWidgetSize");
+}
+
+void SAudioSliderBase::SetUnitsText(const FText Units)
+{
+	UnitsText->SetText(Units);
+}
+
+void SAudioSliderBase::SetTextReadOnly(const bool bIsReadOnly)
+{
+	UnitsText->SetIsReadOnly(bIsReadOnly);
+	ValueText->SetIsReadOnly(bIsReadOnly);
+}
+
+void SAudioSliderBase::SetAlwaysShowLabel(const bool bAlwaysShowLabel)
+{
+	AlwaysShowLabel = bAlwaysShowLabel;
+}
+
+const float SAudioSliderBase::GetOutputValue(const float LinValue)
+{
+	return LinValue;
+}
+
+const float SAudioSliderBase::GetLinValue(const float OutputValue)
+{
+	return OutputValue;
+}
+
+void SAudioSliderBase::SetSliderBackgroundColor(FSlateColor InSliderBackgroundColor)
+{
+	SetAttribute(SliderBackgroundColor, TAttribute<FSlateColor>(InSliderBackgroundColor), EInvalidateWidgetReason::Paint);
+	SliderBackgroundRectangleImage->SetColorAndOpacity(SliderBackgroundColor);
+	SliderBackgroundTopCapImage->SetColorAndOpacity(SliderBackgroundColor);
+	SliderBackgroundBottomCapImage->SetColorAndOpacity(SliderBackgroundColor);
+}
+
+void SAudioSliderBase::SetSliderBarColor(FSlateColor InSliderBarColor)
+{
+	SetAttribute(SliderBarColor, TAttribute<FSlateColor>(InSliderBarColor), EInvalidateWidgetReason::Paint);
+	SliderBarRectangleImage->SetColorAndOpacity(SliderBarColor);
+	SliderBarTopCapImage->SetColorAndOpacity(SliderBarColor);
+	SliderBarBottomCapImage->SetColorAndOpacity(SliderBarColor);
+}
+
+void SAudioSliderBase::SetSliderThumbColor(FSlateColor InSliderThumbColor)
+{
+	SetAttribute(SliderThumbColor, TAttribute<FSlateColor>(InSliderThumbColor), EInvalidateWidgetReason::Paint);
+	Slider->SetSliderHandleColor(SliderThumbColor.Get());
+}
+
+void SAudioSliderBase::SetLabelBackgroundColor(FSlateColor InLabelBackgroundColor)
+{
+	SetAttribute(LabelBackgroundColor, TAttribute<FSlateColor>(InLabelBackgroundColor), EInvalidateWidgetReason::Paint);
+	LabelBackgroundImage->SetColorAndOpacity(LabelBackgroundColor);
+}
+
+void SAudioSliderBase::SetWidgetBackgroundColor(FSlateColor InWidgetBackgroundColor)
+{
+	SetAttribute(WidgetBackgroundColor, TAttribute<FSlateColor>(InWidgetBackgroundColor), EInvalidateWidgetReason::Paint);
+	WidgetBackgroundImage->SetColorAndOpacity(WidgetBackgroundColor);
+}
+
+// SAudioSlider
+SAudioSlider::SAudioSlider()
+{
+}
+
+void SAudioSlider::Construct(const SAudioSliderBase::FArguments& InDeclaration)
+{
+	SAudioSliderBase::Construct(InDeclaration);
+}
+
+void SAudioSlider::SetLinToOutputCurve(const TWeakObjectPtr<const UCurveFloat> InLinToOutputCurve)
+{
+	LinToOutputCurve = InLinToOutputCurve;
+	ValueText->SetText(FText::AsNumber(GetOutputValue(ValueAttribute.Get())));
+}
+
+void SAudioSlider::SetOutputToLinCurve(const TWeakObjectPtr<const UCurveFloat> InOutputToLinCurve)
+{
+	OutputToLinCurve = InOutputToLinCurve;
+}
+
+const TWeakObjectPtr<const UCurveFloat> SAudioSlider::GetOutputToLinCurve()
+{
+	return OutputToLinCurve;
+}
+
+const TWeakObjectPtr<const UCurveFloat> SAudioSlider::GetLinToOutputCurve()
+{
+	return LinToOutputCurve;
+}
+
+const float SAudioSlider::GetOutputValue(const float LinValue)
+{
+	return LinToOutputCurve->IsValidLowLevel() ? LinToOutputCurve->GetFloatValue(LinValue) : LinValue;
+}
+
+const float SAudioSlider::GetLinValue(const float OutputValue)
+{
+	return OutputToLinCurve->IsValidLowLevel() ? OutputToLinCurve->GetFloatValue(OutputValue) : OutputValue;
+}
+
+// SAudioVolumeSlider
+SAudioVolumeSlider::SAudioVolumeSlider()
+{
+}
+
+void SAudioVolumeSlider::Construct(const SAudioSlider::FArguments& InDeclaration)
+{
+	SAudioSlider::Construct(InDeclaration);
+
+	SetLinToOutputCurve(LoadObject<UCurveFloat>(NULL, TEXT("/AudioWidgets/AudioControlCurves/AudioControl_LinDbCurveDefault"), NULL, LOAD_None, NULL));
+	SetOutputToLinCurve(LoadObject<UCurveFloat>(NULL, TEXT("/AudioWidgets/AudioControlCurves/AudioControl_DbLinCurveDefault"), NULL, LOAD_None, NULL));
+	SetUnitsText(FText::FromString("dB"));
+}
+
+// SAudioFrequencySlider
+const FVector2D SAudioFrequencySlider::LinearRange = FVector2D(0.0f, 1.0f);
+
+SAudioFrequencySlider::SAudioFrequencySlider()
+{
+}
+
+void SAudioFrequencySlider::Construct(const SAudioSlider::FArguments& InDeclaration)
+{
+	SAudioSliderBase::Construct(InDeclaration);
+	SetUnitsText(FText::FromString("Hz"));
+}
+
+void SAudioFrequencySlider::SetOutputRange(const FVector2D Range)
+{
+	OutputRange = Range;
+}
+
+const float SAudioFrequencySlider::GetOutputValue(const float LinValue)
+{
+	return Audio::GetLogFrequencyClamped(LinValue, LinearRange, OutputRange);
+}
+
+const float SAudioFrequencySlider::GetLinValue(const float OutputValue)
+{
+	return Audio::GetLinearFrequencyClamped(OutputValue, LinearRange, OutputRange);
+}
diff --git a/Engine/Plugins/Runtime/AudioWidgets/Source/AudioWidgets/Public/AudioSlider.h b/Engine/Plugins/Runtime/AudioWidgets/Source/AudioWidgets/Public/AudioSlider.h
new file mode 100644
--- /dev/null
+++ b/Engine/Plugins/Runtime/AudioWidgets/Source/AudioWidgets/Public/AudioSlider.h
@@ -0,0 +1,198 @@
+// Copyright Epic Games, Inc. All Rights Reserved.
+
+#pragma once
+
+#include "AudioSliderStyle.h"
+#include "Components/Widget.h"
+#include "CoreMinimal.h"
+#include "Curves/CurveFloat.h"
+#include "Styling/SlateTypes.h"
+#include "UObject/ObjectMacros.h"
+#include "Widgets/SWidget.h"
+#include "AudioSlider.generated.h"
+
+class SAudioSliderBase;
+
+DECLARE_DYNAMIC_MULTICAST_DELEGATE_OneParam(FOnFloatValueChangedEvent, float, Value);
+
+/**
+ * An audio slider widget. 
+ */
+UCLASS(Abstract)
+class AUDIOWIDGETS_API UAudioSliderBase: public UWidget
+{
+	GENERATED_UCLASS_BODY()
+
+public:
+	/** The linear value. */
+	UPROPERTY(EditAnywhere, Category = Appearance, meta = (UIMin = "0", UIMax = "1"))
+	float Value;
+
+	/** Whether to show text label. */
+	UPROPERTY(EditAnywhere, Category = Appearance)
+	bool AlwaysShowLabel;
+
+	/** A bindable delegate to allow logic to drive the value of the widget */
+	UPROPERTY()
+	FGetFloat ValueDelegate;
+
+	/** The color to draw the label background. */
+	UPROPERTY(EditAnywhere, Category = Appearance)
+	FLinearColor LabelBackgroundColor;
+
+	/** A bindable delegate for the LabelBackgroundColor. */
+	UPROPERTY()
+	FGetLinearColor LabelBackgroundColorDelegate;
+
+	/** The color to draw the slider background. */
+	UPROPERTY(EditAnywhere, Category = Appearance)
+	FLinearColor SliderBackgroundColor;
+
+	/** A bindable delegate for the SliderBackgroundColor. */
+	UPROPERTY()
+	FGetLinearColor SliderBackgroundColorDelegate;
+
+	/** The color to draw the slider bar. */
+	UPROPERTY(EditAnywhere, Category = Appearance)
+	FLinearColor SliderBarColor;
+
+	/** A bindable delegate for the SliderBarColor. */
+	UPROPERTY()
+	FGetLinearColor SliderBarColorDelegate;
+
+	/** The color to draw the slider thumb. */
+	UPROPERTY(EditAnywhere, Category = Appearance)
+	FLinearColor SliderThumbColor;
+
+	/** A bindable delegate for the SliderThumbColor. */
+	UPROPERTY()
+	FGetLinearColor SliderThumbColorDelegate;
+
+	/** The color to draw the widget background. */
+	UPROPERTY(EditAnywhere, Category = Appearance)
+	FLinearColor WidgetBackgroundColor;
+
+	/** A bindable delegate for the WidgetBackgroundColor. */
+	UPROPERTY()
+	FGetLinearColor WidgetBackgroundColorDelegate;
+
+public:
+	/** Set the units text (ex. "dB") */
+	UFUNCTION(BlueprintCallable, Category = "Behavior")
+	void SetUnitsText(const FText UnitsText);
+
+	/** Get output value from linear based on internal lin to output mapping. */
+	UFUNCTION(BlueprintCallable, Category = "Behavior")
+	float GetOutputValue(const float LinValue);
+
+	/** Get linear value from output based on internal lin to output mapping. */
+	UFUNCTION(BlueprintCallable, Category = "Behavior")
+	float GetLinValue(const float OutputValue);
+
+	/** Sets whether editable value text is read only or not. */
+ 	UFUNCTION(BlueprintCallable, Category = "Behavior")
+ 	void SetTextReadOnly(const bool bIsReadOnly);
+
+	/** Sets whether editable value text is read only or not. */
+	UFUNCTION(BlueprintCallable, Category = "Appearance")
+	void SetAlwaysShowLabel(const bool bSetAlwaysShowLabel);
+	
+	/** Called when the value is changed by slider or typing. */
+	UPROPERTY(BlueprintAssignable, Category = "Widget Event")
+	FOnFloatValueChangedEvent OnValueChanged;
+
+	/** Sets the label background color */
+	UFUNCTION(BlueprintCallable, Category = "Appearance")
+	void SetLabelBackgroundColor(FLinearColor InValue);
+
+	/** Sets the slider background color */
+	UFUNCTION(BlueprintCallable, Category = "Appearance")
+	void SetSliderBackgroundColor(FLinearColor InValue);
+
+	/** Sets the slider bar color */
+	UFUNCTION(BlueprintCallable, Category = "Appearance")
+	void SetSliderBarColor(FLinearColor InValue);
+
+	/** Sets the slider thumb color */
+	UFUNCTION(BlueprintCallable, Category = "Appearance")
+	void SetSliderThumbColor(FLinearColor InValue);
+
+	/** Sets the widget background color */
+	UFUNCTION(BlueprintCallable, Category = "Appearance")
+	void SetWidgetBackgroundColor(FLinearColor InValue);
+
+	// UWidget interface
+	virtual void SynchronizeProperties() override;
+	// End of UWidget interface
+
+	// UVisual interface
+	virtual void ReleaseSlateResources(bool bReleaseChildren) override;
+	// End of UVisual interface
+
+#if WITH_EDITOR
+	virtual const FText GetPaletteCategory() override;
+#endif
+
+protected:
+	/** Native Slate Widget */
+	TSharedPtr<SAudioSliderBase> MyAudioSlider;
+
+	void HandleOnValueChanged(float InValue);
+
+	// UWidget interface
+	virtual TSharedRef<SWidget> RebuildWidget();
+	// End of UWidget interface
+
+	PROPERTY_BINDING_IMPLEMENTATION(float, Value);
+	PROPERTY_BINDING_IMPLEMENTATION(FLinearColor, LabelBackgroundColor);
+	PROPERTY_BINDING_IMPLEMENTATION(FLinearColor, SliderBackgroundColor);
+	PROPERTY_BINDING_IMPLEMENTATION(FLinearColor, SliderBarColor);
+	PROPERTY_BINDING_IMPLEMENTATION(FLinearColor, SliderThumbColor);
+	PROPERTY_BINDING_IMPLEMENTATION(FLinearColor, WidgetBackgroundColor);
+};
+
+/**
+ * An audio slider widget with customizable curves.
+ */
+UCLASS()
+class AUDIOWIDGETS_API UAudioSlider : public UAudioSliderBase
+{
+	GENERATED_UCLASS_BODY()
+
+	/** Curves for mapping linear to output values. */
+	UPROPERTY(EditAnywhere, BlueprintReadWrite, Category = "Behavior")
+	TWeakObjectPtr<const UCurveFloat> LinToOutputCurve;
+
+	UPROPERTY(EditAnywhere, BlueprintReadWrite, Category = "Behavior")
+	TWeakObjectPtr <const UCurveFloat> OutputToLinCurve;
+
+protected:
+	virtual void SynchronizeProperties() override;
+	virtual TSharedRef<SWidget> RebuildWidget() override;
+};
+
+/**
+ * An audio slider widget with default customizable curves for volume (dB).
+ */
+UCLASS()
+class AUDIOWIDGETS_API UAudioVolumeSlider : public UAudioSlider
+{
+	GENERATED_UCLASS_BODY()
+protected:
+	virtual TSharedRef<SWidget> RebuildWidget() override;
+};
+
+/**
+ * An audio slider widget, for use with frequency. 
+ */
+UCLASS()
+class AUDIOWIDGETS_API UAudioFrequencySlider : public UAudioSliderBase
+{
+	GENERATED_UCLASS_BODY()
+	
+	/** Frequency output range */
+	UPROPERTY(EditAnywhere, Category = Behavior)
+	FVector2D OutputRange = FVector2D(20.0f, 20000.0f);
+protected:
+	virtual TSharedRef<SWidget> RebuildWidget() override;
+};
diff --git a/Engine/Plugins/Runtime/AudioWidgets/Source/AudioWidgets/Public/AudioSliderStyle.h b/Engine/Plugins/Runtime/AudioWidgets/Source/AudioWidgets/Public/AudioSliderStyle.h
new file mode 100644
--- /dev/null
+++ b/Engine/Plugins/Runtime/AudioWidgets/Source/AudioWidgets/Public/AudioSliderStyle.h
@@ -0,0 +1,20 @@
+// Copyright Epic Games, Inc. All Rights Reserved.
+
+#pragma once
+
+#include "Brushes/SlateImageBrush.h"
+#include "CoreMinimal.h"
+#include "Styling/SlateStyle.h"
+#include "Styling/SlateStyleMacros.h"
+#include "Styling/SlateStyleRegistry.h"
+#include "Styling/StyleDefaults.h"
+#include "UObject/ObjectMacros.h"
+
+/**
+ * Represents the appearance of an SAudioSlider
+ */
+class FAudioSliderStyle : public FSlateStyleSet
+{
+public: 
+	FAudioSliderStyle();
+};
diff --git a/Engine/Plugins/Runtime/AudioWidgets/Source/AudioWidgets/Public/AudioWidgets.h b/Engine/Plugins/Runtime/AudioWidgets/Source/AudioWidgets/Public/AudioWidgets.h
--- a/Engine/Plugins/Runtime/AudioWidgets/Source/AudioWidgets/Public/AudioWidgets.h
+++ b/Engine/Plugins/Runtime/AudioWidgets/Source/AudioWidgets/Public/AudioWidgets.h
@@ -2,6 +2,7 @@
 
 #pragma once
 
+#include "AudioSliderStyle.h"
 #include "CoreMinimal.h"
 #include "Modules/ModuleManager.h"
 
@@ -12,4 +13,6 @@ public:
 	/** IModuleInterface implementation */
 	virtual void StartupModule() override;
 	virtual void ShutdownModule() override;
+
+	TSharedPtr<FAudioSliderStyle> AudioSliderStyleSet;
 };
diff --git a/Engine/Plugins/Runtime/AudioWidgets/Source/AudioWidgets/Public/SAudioSlider.h b/Engine/Plugins/Runtime/AudioWidgets/Source/AudioWidgets/Public/SAudioSlider.h
new file mode 100644
--- /dev/null
+++ b/Engine/Plugins/Runtime/AudioWidgets/Source/AudioWidgets/Public/SAudioSlider.h
@@ -0,0 +1,171 @@
+// Copyright Epic Games, Inc. All Rights Reserved.
+
+#pragma once
+
+#include "AudioSliderStyle.h"
+#include "CoreMinimal.h"
+#include "Curves/CurveFloat.h"
+#include "Framework/SlateDelegates.h"
+#include "Misc/Attribute.h"
+#include "Styling/CoreStyle.h"
+#include "Styling/SlateTypes.h"
+#include "Styling/SlateWidgetStyleAsset.h"
+#include "Widgets/DeclarativeSyntaxSupport.h"
+#include "Widgets/Images/SImage.h"
+#include "Widgets/Input/SEditableText.h"
+#include "Widgets/Input/SSlider.h"
+#include "Widgets/Layout/SConstraintCanvas.h"
+
+/**
+ * Slate audio sliders that wrap SSlider and provides additional audio specific functionality.
+ * This is a nativized version of the previous Audio Fader widget. 
+ */
+class AUDIOWIDGETS_API SAudioSliderBase
+	: public SCompoundWidget 
+{
+public:
+	SLATE_BEGIN_ARGS(SAudioSliderBase)
+	{
+		_Value = 0.0f;
+		_AlwaysShowLabel = false;
+
+		const ISlateStyle* AudioSliderStyle = FSlateStyleRegistry::FindSlateStyle("AudioSliderStyle");
+		ensure(AudioSliderStyle);
+
+		_LabelBackgroundColor = AudioSliderStyle->GetColor("AudioSlider.DefaultBackgroundColor");
+		_SliderBackgroundColor = AudioSliderStyle->GetColor("AudioSlider.DefaultBackgroundColor");
+		_SliderBarColor = AudioSliderStyle->GetColor("AudioSlider.DefaultBarColor");
+		_SliderThumbColor = AudioSliderStyle->GetColor("AudioSlider.DefaultThumbColor");
+		_WidgetBackgroundColor = AudioSliderStyle->GetColor("AudioSlider.DefaultWidgetBackgroundColor");
+	}
+		/** A value representing the audio slider value. */
+		SLATE_ATTRIBUTE(float, Value)
+
+		/** Whether the text label is always shown or only on hover. */
+		SLATE_ATTRIBUTE(bool, AlwaysShowLabel)
+
+		/** The color to draw the label background in. */
+		SLATE_ATTRIBUTE(FSlateColor, LabelBackgroundColor)
+
+		/** The color to draw the slider background in. */
+		SLATE_ATTRIBUTE(FSlateColor, SliderBackgroundColor)
+
+		/** The color to draw the slider bar in. */
+		SLATE_ATTRIBUTE(FSlateColor, SliderBarColor)
+
+		/** The color to draw the slider thumb in. */
+		SLATE_ATTRIBUTE(FSlateColor, SliderThumbColor)
+		
+		/** The color to draw the widget background in. */
+		SLATE_ATTRIBUTE(FSlateColor, WidgetBackgroundColor)
+
+		/** Called when the value is changed by slider or typing */
+		SLATE_EVENT(FOnFloatValueChanged, OnValueChanged)
+
+	SLATE_END_ARGS()
+
+	SAudioSliderBase();
+
+	/**
+	 * Construct the widget.
+	 *
+	 * @param InDeclaration A declaration from which to construct the widget.
+	 */
+	virtual void Construct(const SAudioSliderBase::FArguments& InDeclaration);
+	virtual const float GetOutputValue(const float LinValue);
+	virtual const float GetLinValue(const float OutputValue);
+	
+	FVector2D ComputeDesiredSize(float) const;
+	void SetUnitsText(const FText Units);
+	/**
+	*  Set whether text label read only or editable.
+	*/
+	void SetTextReadOnly(const bool bIsReadOnly);
+	/**
+	 * Set whether the text label is always shown or only on hover.
+	 */
+	void SetAlwaysShowLabel(const bool bAlwaysShowLabel);
+	void SetSliderBackgroundColor(FSlateColor InSliderBackgroundColor);
+	void SetSliderBarColor(FSlateColor InSliderBarColor);
+	void SetSliderThumbColor(FSlateColor InSliderThumbColor);
+	void SetLabelBackgroundColor(FSlateColor InLabelBackgroundColor);
+	void SetWidgetBackgroundColor(FSlateColor InWidgetBackgroundColor);
+
+protected:
+	// Holds the slider's current linear value, from 0.0 - 1.0f
+	TAttribute<float> ValueAttribute;
+	// Whether the text label is always shown or only on hover
+	TAttribute<bool> AlwaysShowLabel;
+
+	// Various colors 
+	TAttribute<FSlateColor> LabelBackgroundColor;
+	TAttribute<FSlateColor> SliderBackgroundColor;
+	TAttribute<FSlateColor> SliderBarColor;
+	TAttribute<FSlateColor> SliderThumbColor;
+	TAttribute<FSlateColor> WidgetBackgroundColor;
+
+	// Widget components
+	TSharedPtr<SSlider> Slider;
+	TSharedPtr<SEditableText> ValueText;
+	TSharedPtr<SEditableText> UnitsText;
+	TSharedPtr<SImage> LabelBackgroundImage;
+	TSharedPtr<SImage> SliderBackgroundTopCapImage;
+	TSharedPtr<SImage> SliderBackgroundBottomCapImage;
+	TSharedPtr<SImage> SliderBackgroundRectangleImage;
+	TSharedPtr<SImage> SliderBarTopCapImage;
+	TSharedPtr<SImage> SliderBarBottomCapImage;
+	TSharedPtr<SImage> SliderBarRectangleImage;
+	TSharedPtr<SImage> WidgetBackgroundImage;
+
+	// Holds a delegate that is executed when the slider's value changed.
+	FOnFloatValueChanged OnValueChanged;
+};
+
+/* 
+*	An Audio Slider widget with customizable curves. 
+*/
+class AUDIOWIDGETS_API SAudioSlider
+	: public SAudioSliderBase
+{
+public:
+	SAudioSlider();
+	virtual void Construct(const SAudioSliderBase::FArguments& InDeclaration);
+	void SetLinToOutputCurve(const TWeakObjectPtr<const UCurveFloat> LinToOutputCurve);
+	void SetOutputToLinCurve(const TWeakObjectPtr<const UCurveFloat> OutputToLinCurve);
+	const TWeakObjectPtr<const UCurveFloat> GetOutputToLinCurve();
+	const TWeakObjectPtr<const UCurveFloat> GetLinToOutputCurve();
+	const float GetOutputValue(const float LinValue);
+	const float GetLinValue(const float OutputValue);
+protected:
+	// Curves for mapping linear (0.0 - 1.0) to output (ex. dB for volume)  
+	TWeakObjectPtr<const UCurveFloat> LinToOutputCurve = nullptr;
+	TWeakObjectPtr<const UCurveFloat> OutputToLinCurve = nullptr;
+};
+
+/*
+* An Audio Curve Slider widget with default customizable curves for volume (dB). 
+*/
+class AUDIOWIDGETS_API SAudioVolumeSlider
+	: public SAudioSlider
+{
+public:
+	SAudioVolumeSlider();
+	void Construct(const SAudioSliderBase::FArguments& InDeclaration);
+};
+
+/*
+* An Audio Slider widget intended to be used for frequency output, with output frequency range but no customizable curves. 
+*/
+class AUDIOWIDGETS_API SAudioFrequencySlider
+	: public SAudioSliderBase
+{
+public:
+	SAudioFrequencySlider();
+	void Construct(const SAudioSlider::FArguments& InDeclaration);
+	void SetOutputRange(const FVector2D Range);
+	const float GetOutputValue(const float LinValue);
+	const float GetLinValue(const float OutputValue);
+protected:
+	FVector2D OutputRange;
+	static const FVector2D LinearRange;
+};
