Audio Slider widget - minor fixes for static analysis and non unity builds (includes, nullptr ensures, virtual destructors)
#pf 612695a3d673c400011a297d

diff --git a/Engine/Plugins/Runtime/AudioWidgets/Source/AudioWidgets/Private/AudioSliderStyle.cpp b/Engine/Plugins/Runtime/AudioWidgets/Source/AudioWidgets/Private/AudioSliderStyle.cpp
--- a/Engine/Plugins/Runtime/AudioWidgets/Source/AudioWidgets/Private/AudioSliderStyle.cpp
+++ b/Engine/Plugins/Runtime/AudioWidgets/Source/AudioWidgets/Private/AudioSliderStyle.cpp
@@ -1,6 +1,7 @@
 // Copyright Epic Games, Inc. All Rights Reserved.
 
 #include "AudioSliderStyle.h"
+#include "Styling/SlateTypes.h"
 
 FAudioSliderStyle::FAudioSliderStyle()
 	: FSlateStyleSet("AudioSliderStyle")
diff --git a/Engine/Plugins/Runtime/AudioWidgets/Source/AudioWidgets/Private/SAudioSlider.cpp b/Engine/Plugins/Runtime/AudioWidgets/Source/AudioWidgets/Private/SAudioSlider.cpp
--- a/Engine/Plugins/Runtime/AudioWidgets/Source/AudioWidgets/Private/SAudioSlider.cpp
+++ b/Engine/Plugins/Runtime/AudioWidgets/Source/AudioWidgets/Private/SAudioSlider.cpp
@@ -27,7 +27,10 @@ void SAudioSliderBase::Construct(const SAudioSliderBase::FArguments& InDeclarati
 	WidgetBackgroundColor = InDeclaration._WidgetBackgroundColor;
 
 	const ISlateStyle* AudioSliderStyle = FSlateStyleRegistry::FindSlateStyle("AudioSliderStyle");
-	ensure(AudioSliderStyle);
+	if (!ensure(AudioSliderStyle))
+	{
+		return;
+	}
 
 	ChildSlot
 	[
@@ -72,11 +75,9 @@ void SAudioSliderBase::Construct(const SAudioSliderBase::FArguments& InDeclarati
 				.Alignment(FVector2D(0.0f, 0.5f))
 				.Offset(FMargin(0.0f, 0.0f, 0.0f, 32.0f))
 				.AutoSize(true)
-				//.ZOrder(3)
 				[
 					SNew(SHorizontalBox)
 					+ SHorizontalBox::Slot()
-					//.AutoWidth()
 					.FillWidth(0.5)
 					.Padding(4.0f, 0.0f)
 					.HAlign(HAlign_Right)
@@ -215,9 +216,7 @@ void SAudioSliderBase::Construct(const SAudioSliderBase::FArguments& InDeclarati
 FVector2D SAudioSliderBase::ComputeDesiredSize(float) const
 {
 	const ISlateStyle* AudioSliderStyle = FSlateStyleRegistry::FindSlateStyle("AudioSliderStyle");
-	ensure(AudioSliderStyle);
-
-	return AudioSliderStyle->GetVector("AudioSlider.DesiredWidgetSize");
+	return ensure(AudioSliderStyle) ? AudioSliderStyle->GetVector("AudioSlider.DesiredWidgetSize") : FVector2D(26.0f, 477.0f);
 }
 
 void SAudioSliderBase::SetUnitsText(const FText Units)
diff --git a/Engine/Plugins/Runtime/AudioWidgets/Source/AudioWidgets/Public/SAudioSlider.h b/Engine/Plugins/Runtime/AudioWidgets/Source/AudioWidgets/Public/SAudioSlider.h
--- a/Engine/Plugins/Runtime/AudioWidgets/Source/AudioWidgets/Public/SAudioSlider.h
+++ b/Engine/Plugins/Runtime/AudioWidgets/Source/AudioWidgets/Public/SAudioSlider.h
@@ -11,6 +11,7 @@
 #include "Styling/SlateTypes.h"
 #include "Styling/SlateWidgetStyleAsset.h"
 #include "Widgets/DeclarativeSyntaxSupport.h"
+#include "Widgets/SCompoundWidget.h"
 #include "Widgets/Images/SImage.h"
 #include "Widgets/Input/SEditableText.h"
 #include "Widgets/Input/SSlider.h"
@@ -30,13 +31,14 @@ public:
 		_AlwaysShowLabel = false;
 
 		const ISlateStyle* AudioSliderStyle = FSlateStyleRegistry::FindSlateStyle("AudioSliderStyle");
-		ensure(AudioSliderStyle);
-
-		_LabelBackgroundColor = AudioSliderStyle->GetColor("AudioSlider.DefaultBackgroundColor");
-		_SliderBackgroundColor = AudioSliderStyle->GetColor("AudioSlider.DefaultBackgroundColor");
-		_SliderBarColor = AudioSliderStyle->GetColor("AudioSlider.DefaultBarColor");
-		_SliderThumbColor = AudioSliderStyle->GetColor("AudioSlider.DefaultThumbColor");
-		_WidgetBackgroundColor = AudioSliderStyle->GetColor("AudioSlider.DefaultWidgetBackgroundColor");
+		if (ensure(AudioSliderStyle))
+		{
+			_LabelBackgroundColor = AudioSliderStyle->GetColor("AudioSlider.DefaultBackgroundColor");
+			_SliderBackgroundColor = AudioSliderStyle->GetColor("AudioSlider.DefaultBackgroundColor");
+			_SliderBarColor = AudioSliderStyle->GetColor("AudioSlider.DefaultBarColor");
+			_SliderThumbColor = AudioSliderStyle->GetColor("AudioSlider.DefaultThumbColor");
+			_WidgetBackgroundColor = AudioSliderStyle->GetColor("AudioSlider.DefaultWidgetBackgroundColor");
+		}
 	}
 		/** A value representing the audio slider value. */
 		SLATE_ATTRIBUTE(float, Value)
@@ -65,6 +67,7 @@ public:
 	SLATE_END_ARGS()
 
 	SAudioSliderBase();
+	virtual ~SAudioSliderBase() {};
 
 	/**
 	 * Construct the widget.
@@ -129,6 +132,7 @@ class AUDIOWIDGETS_API SAudioSlider
 {
 public:
 	SAudioSlider();
+	virtual ~SAudioSlider() {};
 	virtual void Construct(const SAudioSliderBase::FArguments& InDeclaration);
 	void SetLinToOutputCurve(const TWeakObjectPtr<const UCurveFloat> LinToOutputCurve);
 	void SetOutputToLinCurve(const TWeakObjectPtr<const UCurveFloat> OutputToLinCurve);
