Fix various MetaSound Input Widget and Default Clamp issues
- added option for volume slider to output linear (0-1) value while displaying dB value
- edge case handling for volume and frequency ranges causing widgets to get stuck
- fix functionality and details panel for clamp and other input float properties
- don't show widgets or widget options on readonly graphs

diff --git a/Engine/Plugins/Runtime/AudioWidgets/Source/AudioWidgets/Private/AudioRadialSlider.cpp b/Engine/Plugins/Runtime/AudioWidgets/Source/AudioWidgets/Private/AudioRadialSlider.cpp
--- a/Engine/Plugins/Runtime/AudioWidgets/Source/AudioWidgets/Private/AudioRadialSlider.cpp
+++ b/Engine/Plugins/Runtime/AudioWidgets/Source/AudioWidgets/Private/AudioRadialSlider.cpp
@@ -245,7 +245,7 @@ UAudioFrequencyRadialSlider::UAudioFrequencyRadialSlider(const FObjectInitialize
 	: Super(ObjectInitializer)
 {
 	Super::SetUnitsText(FText::FromString("Hz"));
-	OutputRange = FVector2D(20.0f, 20000.0f);
+	OutputRange = FVector2D(MIN_FILTER_FREQUENCY, MAX_FILTER_FREQUENCY);
 }
 
 TSharedRef<SWidget> UAudioFrequencyRadialSlider::RebuildWidget()
diff --git a/Engine/Plugins/Runtime/AudioWidgets/Source/AudioWidgets/Private/AudioSlider.cpp b/Engine/Plugins/Runtime/AudioWidgets/Source/AudioWidgets/Private/AudioSlider.cpp
--- a/Engine/Plugins/Runtime/AudioWidgets/Source/AudioWidgets/Private/AudioSlider.cpp
+++ b/Engine/Plugins/Runtime/AudioWidgets/Source/AudioWidgets/Private/AudioSlider.cpp
@@ -253,7 +253,7 @@ TSharedRef<SWidget> UAudioVolumeSlider::RebuildWidget()
 // UAudioFrequencySlider
 UAudioFrequencySlider::UAudioFrequencySlider(const FObjectInitializer& ObjectInitializer)
 	: Super(ObjectInitializer)
-	, OutputRange(FVector2D(20.0f, 20000.0f))
+	, OutputRange(FVector2D(MIN_FILTER_FREQUENCY, MAX_FILTER_FREQUENCY))
 {
 	Super::SetUnitsText(FText::FromString("Hz"));
 }
diff --git a/Engine/Plugins/Runtime/AudioWidgets/Source/AudioWidgets/Private/SAudioRadialSlider.cpp b/Engine/Plugins/Runtime/AudioWidgets/Source/AudioWidgets/Private/SAudioRadialSlider.cpp
--- a/Engine/Plugins/Runtime/AudioWidgets/Source/AudioWidgets/Private/SAudioRadialSlider.cpp
+++ b/Engine/Plugins/Runtime/AudioWidgets/Source/AudioWidgets/Private/SAudioRadialSlider.cpp
@@ -41,7 +41,7 @@ void SAudioRadialSlider::Construct(const SAudioRadialSlider::FArguments& InArgs)
 		.OnValueTextCommitted_Lambda([this](const FText& Text, ETextCommit::Type CommitType)
 		{
 			const float OutputValue = FCString::Atof(*Text.ToString());
-			const float LinValue = GetLinValue(OutputValue);
+			const float LinValue = GetLinValueForText(OutputValue);
 			Value.Set(LinValue);
 			RadialSlider->SetValue(LinValue);
 			OnValueChanged.ExecuteIfBound(LinValue);
@@ -52,7 +52,7 @@ void SAudioRadialSlider::Construct(const SAudioRadialSlider::FArguments& InArgs)
 		{
 			Value.Set(InLinValue);
 			OnValueChanged.ExecuteIfBound(InLinValue);
-			const float OutputValue = GetOutputValue(InLinValue);
+			const float OutputValue = GetOutputValueForText(InLinValue);
 			Label->SetValueText(OutputValue);
 		})
 		.UseVerticalDrag(true)
@@ -187,8 +187,8 @@ TSharedRef<SWidgetSwitcher> SAudioRadialSlider::CreateLayoutWidgetSwitcher()
 void SAudioRadialSlider::SetValue(float LinValue)
 {
 	Value.Set(LinValue);
-	const float OutputValue = GetOutputValue(LinValue);
-	Label->SetValueText(OutputValue);
+	const float OutputValueText = GetOutputValueForText(LinValue);
+	Label->SetValueText(OutputValueText);
 	RadialSlider->SetValue(LinValue);
 }
 
@@ -204,12 +204,26 @@ const float SAudioRadialSlider::GetOutputValue(const float LinValue)
 
 void SAudioRadialSlider::SetOutputRange(const FVector2D Range)
 {
-	if (Range.Y > Range.X)
-	{
-		OutputRange = Range;
-		SetValue(FMath::Clamp(Value.Get(), Range.X, Range.Y));
-		Label->UpdateValueTextWidth(Range);
-	}
+	OutputRange = Range;
+	// if Range.Y < Range.X, set Range.X to Range.Y
+	OutputRange.X = FMath::Min(Range.X, Range.Y);
+
+	const float OutputValue = GetOutputValue(Value.Get());
+	const float ClampedOutputValue = FMath::Clamp(OutputValue, OutputRange.X, OutputRange.Y);
+	const float LinValue = GetLinValue(ClampedOutputValue);
+	SetValue(LinValue);
+
+	Label->UpdateValueTextWidth(OutputRange);
+}
+
+const float SAudioRadialSlider::GetOutputValueForText(const float LinValue)
+{
+	return GetOutputValue(LinValue);
+}
+
+const float SAudioRadialSlider::GetLinValueForText(const float OutputValue)
+{
+	return GetLinValue(OutputValue);
 }
 
 void SAudioRadialSlider::SetLabelBackgroundColor(FSlateColor InColor)
@@ -249,6 +263,9 @@ void SAudioRadialSlider::SetSliderThickness(const float Thickness)
 }
 
 // SAudioVolumeRadialSlider
+const float SAudioVolumeRadialSlider::MinDbValue = -160.0f;
+const float SAudioVolumeRadialSlider::MaxDbValue = 770.0f;
+
 SAudioVolumeRadialSlider::SAudioVolumeRadialSlider()
 {
 }
@@ -257,11 +274,11 @@ void SAudioVolumeRadialSlider::Construct(const SAudioRadialSlider::FArguments& I
 {
 	SAudioRadialSlider::Construct(InArgs);
 	
-	SetOutputRange(FVector2D(-100.0f, 0.0f));
+	SAudioRadialSlider::SetOutputRange(FVector2D(-100.0f, 0.0f));
 	Label->SetUnitsText(FText::FromString("dB"));
 }
 
-const float SAudioVolumeRadialSlider::GetOutputValue(const float LinValue)
+const float SAudioVolumeRadialSlider::GetDbValueFromLin(const float LinValue)
 {
 	// convert from linear 0-1 space to decibel OutputRange that has been converted to linear 
 	const FVector2D LinearSliderRange = FVector2D(Audio::ConvertToLinear(OutputRange.X), Audio::ConvertToLinear(OutputRange.Y));
@@ -271,7 +288,7 @@ const float SAudioVolumeRadialSlider::GetOutputValue(const float LinValue)
 	return FMath::Clamp(OutputValue, OutputRange.X, OutputRange.Y);
 }
 
-const float SAudioVolumeRadialSlider::GetLinValue(const float OutputValue)
+const float SAudioVolumeRadialSlider::GetLinValueFromDb(const float OutputValue)
 {
 	float ClampedValue = FMath::Clamp(OutputValue, OutputRange.X, OutputRange.Y);
 	// convert from decibels to linear
@@ -281,6 +298,54 @@ const float SAudioVolumeRadialSlider::GetLinValue(const float OutputValue)
 	return FMath::GetMappedRangeValueClamped(LinearSliderRange, LinearRange, LinearSliderValue);
 }
 
+const float SAudioVolumeRadialSlider::GetOutputValue(const float LinValue)
+{
+	if (bUseLinearOutput)
+	{
+		return FMath::Clamp(LinValue, LinearRange.X, LinearRange.Y);
+	}
+	else
+	{
+		return GetDbValueFromLin(LinValue);
+	}
+}
+
+const float SAudioVolumeRadialSlider::GetLinValue(const float OutputValue)
+{
+	if (bUseLinearOutput)
+	{
+		return FMath::Clamp(OutputValue, LinearRange.X, LinearRange.Y);
+	}
+	else
+	{
+		return GetLinValueFromDb(OutputValue);
+	}
+}
+
+const float SAudioVolumeRadialSlider::GetOutputValueForText(const float LinValue)
+{
+	return GetDbValueFromLin(LinValue);
+}
+
+const float SAudioVolumeRadialSlider::GetLinValueForText(const float OutputValue)
+{
+	return GetLinValueFromDb(OutputValue);
+}
+
+void SAudioVolumeRadialSlider::SetUseLinearOutput(bool InUseLinearOutput)
+{
+	bUseLinearOutput = InUseLinearOutput;
+}
+
+void SAudioVolumeRadialSlider::SetOutputRange(const FVector2D Range)
+{
+	// if using linear output, output range cannot be changed 
+	if (!bUseLinearOutput)
+	{
+		SAudioRadialSlider::SetOutputRange(FVector2D(FMath::Max(MinDbValue, Range.X), FMath::Min(MaxDbValue, Range.Y)));
+	}
+}
+
 // SAudioFrequencyRadialSlider
 SAudioFrequencyRadialSlider::SAudioFrequencyRadialSlider()
 {
@@ -290,7 +355,7 @@ void SAudioFrequencyRadialSlider::Construct(const SAudioRadialSlider::FArguments
 {
 	SAudioRadialSlider::Construct(InArgs);
 
-	SetOutputRange(FVector2D(20.0f, 20000.0f));
+	SAudioRadialSlider::SetOutputRange(FVector2D(MIN_FILTER_FREQUENCY, MAX_FILTER_FREQUENCY));
 	Label->SetUnitsText(FText::FromString("Hz"));
 }
 
@@ -301,5 +366,14 @@ const float SAudioFrequencyRadialSlider::GetOutputValue(const float LinValue)
 
 const float SAudioFrequencyRadialSlider::GetLinValue(const float OutputValue)
 {
+	// edge case to avoid audio function returning negative value
+	if (FMath::IsNearlyEqual(OutputValue, OutputRange.X))
+	{
+		return LinearRange.X;
+	}
+	if (FMath::IsNearlyEqual(OutputValue, OutputRange.Y))
+	{
+		return LinearRange.Y;
+	}
 	return Audio::GetLinearFrequencyClamped(OutputValue, LinearRange, OutputRange);
 }
diff --git a/Engine/Plugins/Runtime/AudioWidgets/Source/AudioWidgets/Private/SAudioSlider.cpp b/Engine/Plugins/Runtime/AudioWidgets/Source/AudioWidgets/Private/SAudioSlider.cpp
--- a/Engine/Plugins/Runtime/AudioWidgets/Source/AudioWidgets/Private/SAudioSlider.cpp
+++ b/Engine/Plugins/Runtime/AudioWidgets/Source/AudioWidgets/Private/SAudioSlider.cpp
@@ -58,7 +58,7 @@ void SAudioSliderBase::Construct(const SAudioSliderBase::FArguments& InArgs)
 		{
 			ValueAttribute.Set(Value);
 			OnValueChanged.ExecuteIfBound(Value);
-			const float OutputValue = GetOutputValue(Value);
+			const float OutputValue = GetOutputValueForText(Value);
 			Label->SetValueText(OutputValue);
 		})
 		.OnMouseCaptureEnd_Lambda([this]()
@@ -72,7 +72,7 @@ void SAudioSliderBase::Construct(const SAudioSliderBase::FArguments& InArgs)
 		.OnValueTextCommitted_Lambda([this](const FText& Text, ETextCommit::Type CommitType)
 		{
 			const float OutputValue = FCString::Atof(*Text.ToString());
-			const float LinValue = GetLinValue(OutputValue);
+			const float LinValue = GetLinValueForText(OutputValue);
 			ValueAttribute.Set(LinValue);
 			Slider->SetValue(LinValue);
 			OnValueChanged.ExecuteIfBound(LinValue);
@@ -88,8 +88,8 @@ void SAudioSliderBase::Construct(const SAudioSliderBase::FArguments& InArgs)
 void SAudioSliderBase::SetValue(float LinValue)
 {
 	ValueAttribute.Set(LinValue);
-	const float OutputValue = GetOutputValue(LinValue);
-	Label->SetValueText(OutputValue);
+	const float OutputValueForText = GetOutputValueForText(LinValue);
+	Label->SetValueText(OutputValueForText);
 	Slider->SetValue(LinValue);
 }
 
@@ -140,6 +140,16 @@ const float SAudioSliderBase::GetOutputValue(const float LinValue)
 	return FMath::Clamp(LinValue, OutputRange.X, OutputRange.Y);
 }
 
+const float SAudioSliderBase::GetOutputValueForText(const float LinValue)
+{
+	return GetOutputValue(LinValue);
+}
+
+const float SAudioSliderBase::GetLinValueForText(const float OutputValue)
+{
+	return GetLinValue(OutputValue);
+}
+
 const float SAudioSliderBase::GetLinValue(const float OutputValue)
 {
 	return OutputValue;
@@ -173,12 +183,16 @@ void SAudioSliderBase::SetWidgetBackgroundColor(FSlateColor InWidgetBackgroundCo
 
 void SAudioSliderBase::SetOutputRange(const FVector2D Range)
 {
-	if (Range.Y > Range.X)
-	{
-		OutputRange = Range;
-		SetValue(FMath::Clamp(ValueAttribute.Get(), Range.X, Range.Y));
-		Label->UpdateValueTextWidth(Range);
-	}
+	OutputRange = Range;
+	// if Range.Y < Range.X, set Range.X to Range.Y
+	OutputRange.X = FMath::Min(Range.X, Range.Y);
+
+	const float OutputValue = GetOutputValue(ValueAttribute.Get());
+	const float ClampedOutputValue = FMath::Clamp(OutputValue, OutputRange.X, OutputRange.Y);
+	const float LinValue = GetLinValue(ClampedOutputValue);
+	SetValue(LinValue);
+
+	Label->UpdateValueTextWidth(OutputRange);
 }
 
 void SAudioSliderBase::SetLabelBackgroundColor(FSlateColor InColor)
@@ -332,7 +346,7 @@ void SAudioSlider::Construct(const SAudioSliderBase::FArguments& InArgs)
 void SAudioSlider::SetLinToOutputCurve(const TWeakObjectPtr<const UCurveFloat> InLinToOutputCurve)
 {
 	LinToOutputCurve = InLinToOutputCurve;
-	Label->SetValueText(GetOutputValue(ValueAttribute.Get()));
+	Label->SetValueText(GetOutputValueForText(ValueAttribute.Get()));
 }
 
 void SAudioSlider::SetOutputToLinCurve(const TWeakObjectPtr<const UCurveFloat> InOutputToLinCurve)
@@ -371,6 +385,8 @@ const float SAudioSlider::GetLinValue(const float OutputValue)
 }
 
 // SAudioVolumeSlider
+const float SAudioVolumeSlider::MinDbValue = -160.0f;
+const float SAudioVolumeSlider::MaxDbValue = 770.0f;
 SAudioVolumeSlider::SAudioVolumeSlider()
 {
 }
@@ -378,11 +394,11 @@ SAudioVolumeSlider::SAudioVolumeSlider()
 void SAudioVolumeSlider::Construct(const SAudioSlider::FArguments& InArgs)
 {
 	SAudioSliderBase::Construct(InArgs);
-	SetOutputRange(FVector2D(-100.0f, 0.0f));
+	SAudioSliderBase::SetOutputRange(FVector2D(-100.0f, 0.0f));
 	Label->SetUnitsText(FText::FromString("dB"));
 }
 
-const float SAudioVolumeSlider::GetOutputValue(const float LinValue)
+const float SAudioVolumeSlider::GetDbValueFromLin(const float LinValue)
 {
 	// convert from linear 0-1 space to decibel OutputRange that has been converted to linear 
 	const FVector2D LinearSliderRange = FVector2D(Audio::ConvertToLinear(OutputRange.X), Audio::ConvertToLinear(OutputRange.Y));
@@ -392,9 +408,9 @@ const float SAudioVolumeSlider::GetOutputValue(const float LinValue)
 	return FMath::Clamp(OutputValue, OutputRange.X, OutputRange.Y);
 }
 
-const float SAudioVolumeSlider::GetLinValue(const float OutputValue)
+const float SAudioVolumeSlider::GetLinValueFromDb(const float DbValue)
 {
-	float ClampedValue = FMath::Clamp(OutputValue, OutputRange.X, OutputRange.Y);
+	float ClampedValue = FMath::Clamp(DbValue, OutputRange.X, OutputRange.Y);
 	// convert from decibels to linear
 	float LinearSliderValue = Audio::ConvertToLinear(ClampedValue);
 	// convert from decibel OutputRange that has been converted to linear to linear 0-1 space 
@@ -402,6 +418,54 @@ const float SAudioVolumeSlider::GetLinValue(const float OutputValue)
 	return FMath::GetMappedRangeValueClamped(LinearSliderRange, LinearRange, LinearSliderValue);
 }
 
+const float SAudioVolumeSlider::GetOutputValue(const float LinValue)
+{
+	if (bUseLinearOutput)
+	{
+		return FMath::Clamp(LinValue, LinearRange.X, LinearRange.Y);
+	}
+	else
+	{
+		return GetDbValueFromLin(LinValue);
+	}
+}
+
+const float SAudioVolumeSlider::GetLinValue(const float OutputValue)
+{
+	if (bUseLinearOutput)
+	{
+		return FMath::Clamp(OutputValue, LinearRange.X, LinearRange.Y);
+	}
+	else
+	{
+		return GetLinValueFromDb(OutputValue);
+	}
+}
+
+const float SAudioVolumeSlider::GetOutputValueForText(const float LinValue)
+{
+	return GetDbValueFromLin(LinValue);
+}
+
+const float SAudioVolumeSlider::GetLinValueForText(const float OutputValue)
+{
+	return GetLinValueFromDb(OutputValue);
+}
+
+void SAudioVolumeSlider::SetUseLinearOutput(bool InUseLinearOutput)
+{
+	bUseLinearOutput = InUseLinearOutput;
+}
+
+void SAudioVolumeSlider::SetOutputRange(const FVector2D Range)
+{
+	// if using linear output, output range cannot be changed 
+	if (!bUseLinearOutput)
+	{
+		SAudioSliderBase::SetOutputRange(FVector2D(FMath::Max(MinDbValue, Range.X), FMath::Min(MaxDbValue, Range.Y)));
+	}
+}
+
 // SAudioFrequencySlider
 SAudioFrequencySlider::SAudioFrequencySlider()
 {
@@ -410,7 +474,7 @@ SAudioFrequencySlider::SAudioFrequencySlider()
 void SAudioFrequencySlider::Construct(const SAudioSlider::FArguments& InArgs)
 {
 	SAudioSliderBase::Construct(InArgs);
-	SetOutputRange(FVector2D(20.0f, 20000.0f));
+	SetOutputRange(FVector2D(MIN_FILTER_FREQUENCY, MAX_FILTER_FREQUENCY));
 	Label->SetUnitsText(FText::FromString("Hz"));
 }
 
@@ -421,5 +485,14 @@ const float SAudioFrequencySlider::GetOutputValue(const float LinValue)
 
 const float SAudioFrequencySlider::GetLinValue(const float OutputValue)
 {
+	// edge case to avoid audio function returning negative value
+	if (FMath::IsNearlyEqual(OutputValue, OutputRange.X))
+	{
+		return LinearRange.X;
+	}
+	if (FMath::IsNearlyEqual(OutputValue, OutputRange.Y))
+	{
+		return LinearRange.Y;
+	}
 	return Audio::GetLinearFrequencyClamped(OutputValue, LinearRange, OutputRange);
 }
diff --git a/Engine/Plugins/Runtime/AudioWidgets/Source/AudioWidgets/Public/AudioSlider.h b/Engine/Plugins/Runtime/AudioWidgets/Source/AudioWidgets/Public/AudioSlider.h
--- a/Engine/Plugins/Runtime/AudioWidgets/Source/AudioWidgets/Public/AudioSlider.h
+++ b/Engine/Plugins/Runtime/AudioWidgets/Source/AudioWidgets/Public/AudioSlider.h
@@ -2,6 +2,7 @@
 
 #pragma once
 
+#include "AudioDefines.h"
 #include "AudioWidgetsSlateTypes.h"
 #include "Components/Widget.h"
 #include "Curves/CurveFloat.h"
@@ -222,7 +223,7 @@ class AUDIOWIDGETS_API UAudioFrequencySlider : public UAudioSliderBase
 	
 	/** Frequency output range */
 	UPROPERTY(EditAnywhere, Category = Behavior)
-	FVector2D OutputRange = FVector2D(20.0f, 20000.0f);
+	FVector2D OutputRange = FVector2D(MIN_FILTER_FREQUENCY, MAX_FILTER_FREQUENCY);
 protected:
 	virtual TSharedRef<SWidget> RebuildWidget() override;
 };
diff --git a/Engine/Plugins/Runtime/AudioWidgets/Source/AudioWidgets/Public/SAudioRadialSlider.h b/Engine/Plugins/Runtime/AudioWidgets/Source/AudioWidgets/Public/SAudioRadialSlider.h
--- a/Engine/Plugins/Runtime/AudioWidgets/Source/AudioWidgets/Public/SAudioRadialSlider.h
+++ b/Engine/Plugins/Runtime/AudioWidgets/Source/AudioWidgets/Public/SAudioRadialSlider.h
@@ -2,6 +2,7 @@
 
 #pragma once
 
+#include "AudioDefines.h"
 #include "AudioWidgetsSlateTypes.h"
 #include "Curves/CurveFloat.h"
 #include "Framework/SlateDelegates.h"
@@ -107,7 +108,9 @@ public:
 	void SetValue(float LinValue);
 	virtual const float GetOutputValue(const float LinValue);
 	virtual const float GetLinValue(const float OutputValue);
-	void SetOutputRange(const FVector2D Range);
+	virtual const float GetOutputValueForText(const float LinValue);
+	virtual const float GetLinValueForText(const float OutputValue);
+	virtual void SetOutputRange(const FVector2D Range);
 
 	// Text label functions 
 	void SetLabelBackgroundColor(FSlateColor InColor);
@@ -155,6 +158,21 @@ public:
 	void Construct(const SAudioRadialSlider::FArguments& InArgs);
 	const float GetOutputValue(const float LinValue);
 	const float GetLinValue(const float OutputValue);
+	const float GetOutputValueForText(const float LinValue) override;
+	const float GetLinValueForText(const float OutputValue) override;
+	void SetUseLinearOutput(bool InUseLinearOutput);
+	void SetOutputRange(const FVector2D Range) override;
+
+private:
+	// Min/max possible values for output range, derived to avoid Audio::ConvertToLinear/dB functions returning NaN
+	static const float MinDbValue; 
+	static const float MaxDbValue; 
+
+	// Use linear (0 - 1) output value. Only applies to the output value reported by GetOutputValue(); the text displayed will still be in decibels. 
+	bool bUseLinearOutput = true;
+
+	const float GetDbValueFromLin(const float LinValue);
+	const float GetLinValueFromDb(const float DbValue);
 };
 
 /*
diff --git a/Engine/Plugins/Runtime/AudioWidgets/Source/AudioWidgets/Public/SAudioSlider.h b/Engine/Plugins/Runtime/AudioWidgets/Source/AudioWidgets/Public/SAudioSlider.h
--- a/Engine/Plugins/Runtime/AudioWidgets/Source/AudioWidgets/Public/SAudioSlider.h
+++ b/Engine/Plugins/Runtime/AudioWidgets/Source/AudioWidgets/Public/SAudioSlider.h
@@ -93,6 +93,8 @@ public:
 	virtual void Construct(const SAudioSliderBase::FArguments& InDeclaration);
 	virtual const float GetOutputValue(const float LinValue);
 	virtual const float GetLinValue(const float OutputValue);
+	virtual const float GetOutputValueForText(const float LinValue);
+	virtual const float GetLinValueForText(const float OutputValue);
 	
 	/**
 	 * Set the slider's linear (0-1 normalized) value. 
@@ -106,7 +108,7 @@ public:
 	void SetSliderBarColor(FSlateColor InSliderBarColor);
 	void SetSliderThumbColor(FSlateColor InSliderThumbColor);
 	void SetWidgetBackgroundColor(FSlateColor InWidgetBackgroundColor);
-	void SetOutputRange(const FVector2D Range);
+	virtual void SetOutputRange(const FVector2D Range);
 
 	// Text label functions 
 	void SetLabelBackgroundColor(FSlateColor InColor);
@@ -183,8 +185,24 @@ class AUDIOWIDGETS_API SAudioVolumeSlider
 public:
 	SAudioVolumeSlider();
 	void Construct(const SAudioSlider::FArguments& InDeclaration);
+
 	const float GetOutputValue(const float LinValue);
 	const float GetLinValue(const float OutputValue);
+	const float GetOutputValueForText(const float LinValue) override;
+	const float GetLinValueForText(const float OutputValue) override;
+	void SetUseLinearOutput(bool InUseLinearOutput);
+	void SetOutputRange(const FVector2D Range) override;
+
+private:
+	// Min/max possible values for output range, derived to avoid Audio::ConvertToLinear/dB functions returning NaN
+	static const float MinDbValue;
+	static const float MaxDbValue; 
+	
+	// Use linear (0 - 1) output value. Only applies to the output value reported by GetOutputValue(); the text displayed will still be in decibels. 
+	bool bUseLinearOutput = true;
+
+	const float GetDbValueFromLin(const float LinValue);
+	const float GetLinValueFromDb(const float DbValue);
 };
 
 /*
diff --git a/Engine/Plugins/Runtime/Metasound/Source/MetasoundEditor/Private/MetasoundEditorCommands.cpp b/Engine/Plugins/Runtime/Metasound/Source/MetasoundEditor/Private/MetasoundEditorCommands.cpp
--- a/Engine/Plugins/Runtime/Metasound/Source/MetasoundEditor/Private/MetasoundEditorCommands.cpp
+++ b/Engine/Plugins/Runtime/Metasound/Source/MetasoundEditor/Private/MetasoundEditorCommands.cpp
@@ -25,7 +25,7 @@ namespace Metasound
 
 			UI_COMMAND(UpdateNodes, "Update Nodes", "Update selected node(s) that have an available update.", EUserInterfaceActionType::Button, FInputChord());
 
-			UI_COMMAND(ConvertFromPreset, "Convert From Preset", "Convert's this preset to a fully accessible MetaSound.", EUserInterfaceActionType::Button, FInputChord());
+			UI_COMMAND(ConvertFromPreset, "Convert From Preset", "Converts this preset to a fully accessible MetaSound.", EUserInterfaceActionType::Button, FInputChord());
 		}
 	} // namespace Editor
 } // namespace Metasound
diff --git a/Engine/Plugins/Runtime/Metasound/Source/MetasoundEditor/Private/MetasoundEditorGraphMemberDefaults.cpp b/Engine/Plugins/Runtime/Metasound/Source/MetasoundEditor/Private/MetasoundEditorGraphMemberDefaults.cpp
--- a/Engine/Plugins/Runtime/Metasound/Source/MetasoundEditor/Private/MetasoundEditorGraphMemberDefaults.cpp
+++ b/Engine/Plugins/Runtime/Metasound/Source/MetasoundEditor/Private/MetasoundEditorGraphMemberDefaults.cpp
@@ -3,6 +3,7 @@
 
 #include "Algo/Count.h"
 #include "Algo/Transform.h"
+#include "AudioDefines.h"
 #include "AudioParameterControllerInterface.h"
 #include "EdGraph/EdGraphNode.h"
 #include "GraphEditorSettings.h"
@@ -166,49 +167,50 @@ void UMetasoundEditorGraphMemberDefaultFloat::UpdatePreviewInstance(const Metaso
 	InParameterInterface->SetFloatParameter(InParameterName, Default);
 }
 
-void UMetasoundEditorGraphMemberDefaultFloat::PostEditChangeProperty(FPropertyChangedEvent& PropertyChangedEvent)
+void UMetasoundEditorGraphMemberDefaultFloat::PostEditChangeChainProperty(FPropertyChangedChainEvent& PropertyChangedEvent)
 {
 	if (PropertyChangedEvent.GetPropertyName().IsEqual(GET_MEMBER_NAME_CHECKED(UMetasoundEditorGraphMemberDefaultFloat, Default)))
 	{
-		OnDefaultValueChanged.Broadcast(Default);
+		SetDefault(Default);
 	}
 	else if (PropertyChangedEvent.GetPropertyName().IsEqual(GET_MEMBER_NAME_CHECKED(UMetasoundEditorGraphMemberDefaultFloat, WidgetType)) || PropertyChangedEvent.GetPropertyName().IsEqual(GET_MEMBER_NAME_CHECKED(UMetasoundEditorGraphMemberDefaultFloat, WidgetValueType)))
 	{
 		if (WidgetValueType == EMetasoundMemberDefaultWidgetValueType::Linear)
 		{
-			Range = FVector2D(0.0f, 1.0f);
+			SetRange(FVector2D(0.0f, 1.0f));
 		}
 		else if (WidgetValueType == EMetasoundMemberDefaultWidgetValueType::Frequency)
 		{
-			Range = FVector2D(20.0f, 20000.0f);
+			SetRange(FVector2D(MIN_FILTER_FREQUENCY, MAX_FILTER_FREQUENCY));
 		}
 		else if (WidgetValueType == EMetasoundMemberDefaultWidgetValueType::Volume)
 		{
-			Range = FVector2D(-100.0f, 0.0f);
+			SetRange(FVector2D(-100.0f, 0.0f));
 		}
 
-		OnRangeChanged.Broadcast(Range);
 		// If the widget type is changed to none, we need to refresh clamping the value or not, since if the widget was a slider before, the value was clamped
 		OnClampChanged.Broadcast(ClampDefault);
 	}
-	else if (PropertyChangedEvent.MemberProperty->GetFName().IsEqual(GET_MEMBER_NAME_CHECKED(UMetasoundEditorGraphMemberDefaultFloat, Range)))
+	else if (PropertyChangedEvent.PropertyChain.GetActiveMemberNode()->GetValue()->GetFName().IsEqual(GET_MEMBER_NAME_CHECKED(UMetasoundEditorGraphMemberDefaultFloat, Range)))
 	{
 		if (PropertyChangedEvent.ChangeType != EPropertyChangeType::Interactive)
 		{
-			float Min = Range.X;
-			float Max = Range.Y;
-			if (Min < Max)
-			{
-				OnRangeChanged.Broadcast(Range);
-				SetDefault(FMath::Clamp(Default, Min, Max));
-			}
+			// if Range.Y < Range.X, set Range.X to Range.Y
+			Range.X = FMath::Min(Range.X, Range.Y);
+			SetRange(Range);
 		}
 	}
 	else if (PropertyChangedEvent.GetPropertyName().IsEqual(GET_MEMBER_NAME_CHECKED(UMetasoundEditorGraphMemberDefaultFloat, ClampDefault)))
 	{
 		// set range to reasonable limit given current value
-		Range = FVector2D(FMath::Min(0.0f, Default), FMath::Max(0.0f, Default));
-
+		if (FMath::IsNearlyEqual(Default, 0.0f))
+		{
+			SetRange(FVector2D(0.0f, 1.0f));
+		}
+		else
+		{
+			SetRange(FVector2D(FMath::Min(0.0f, Default), FMath::Max(0.0f, Default)));
+		}
 		OnClampChanged.Broadcast(ClampDefault);
 	}
 
@@ -217,6 +219,8 @@ void UMetasoundEditorGraphMemberDefaultFloat::PostEditChangeProperty(FPropertyCh
 	{
 		Member->MarkNodesForRefresh();
 	}
+
+	Super::PostEditChangeChainProperty(PropertyChangedEvent);
 }
 
 void UMetasoundEditorGraphMemberDefaultFloat::SetDefault(const float InDefault)
@@ -235,6 +239,13 @@ FVector2D UMetasoundEditorGraphMemberDefaultFloat::GetRange()
 	return Range;
 }
 
+void UMetasoundEditorGraphMemberDefaultFloat::SetRange(const FVector2D InRange)
+{
+	Range = InRange;
+	OnRangeChanged.Broadcast(InRange);
+	SetDefault(FMath::Clamp(Default, Range.X, Range.Y));
+}
+
 FMetasoundFrontendLiteral UMetasoundEditorGraphMemberDefaultFloatArray::GetDefault() const
 {
 	FMetasoundFrontendLiteral Literal;
diff --git a/Engine/Plugins/Runtime/Metasound/Source/MetasoundEditor/Private/MetasoundNodeDetailCustomization.cpp b/Engine/Plugins/Runtime/Metasound/Source/MetasoundEditor/Private/MetasoundNodeDetailCustomization.cpp
--- a/Engine/Plugins/Runtime/Metasound/Source/MetasoundEditor/Private/MetasoundNodeDetailCustomization.cpp
+++ b/Engine/Plugins/Runtime/Metasound/Source/MetasoundEditor/Private/MetasoundNodeDetailCustomization.cpp
@@ -122,11 +122,21 @@ namespace Metasound
 				DefaultValueHandle = Row->GetPropertyHandle();
 			}
 
+			// Apply the clamp range to the default value if 
+			// 1. not using a widget and ClampDefault is true or 
+			// 2. using a widget but not using a volume widget where VolumeWidgetUseLinearOutput is true 
+			const bool bVolumeWidgetWithLinearOutput = DefaultFloat->WidgetValueType == EMetasoundMemberDefaultWidgetValueType::Volume && DefaultFloat->VolumeWidgetUseLinearOutput;
+			const bool bUsingWidget = DefaultFloat->WidgetType != EMetasoundMemberDefaultWidget::None;
+			const bool bShouldClampDefaultValue = (!bUsingWidget && DefaultFloat->ClampDefault) || (bUsingWidget && !bVolumeWidgetWithLinearOutput);
+
 			Row = DefaultCategoryBuilder->AddExternalObjectProperty(TArray<UObject*>({ DefaultFloat }), GET_MEMBER_NAME_CHECKED(UMetasoundEditorGraphMemberDefaultFloat, ClampDefault));
 			if (ensure(Row))
 			{
-				// If clamping or using slider, clamp default value to given range 
-				if (DefaultFloat->ClampDefault || DefaultFloat->WidgetType != EMetasoundMemberDefaultWidget::None)
+				if (bVolumeWidgetWithLinearOutput)
+				{
+					DefaultFloat->SetRange(FVector2D(0.0f, 1.0f));
+				}
+				if (bShouldClampDefaultValue)
 				{
 					FVector2D Range = DefaultFloat->GetRange();
 					DefaultValueHandle->SetInstanceMetaData("ClampMin", FString::Printf(TEXT("%f"), Range.X));
@@ -150,14 +160,35 @@ namespace Metasound
 						}
 					}
 				});
+
+				if (bShouldClampDefaultValue)
+				{
+					DefaultCategoryBuilder->AddExternalObjectProperty(TArray<UObject*>({ DefaultFloat }), GET_MEMBER_NAME_CHECKED(UMetasoundEditorGraphMemberDefaultFloat, Range));
+				}
+			}
+
+			bool bIsGraphEditable = false;
+			if (const UMetasoundEditorGraphMember* ParentMember = InLiteral.GetParentMember())
+			{
+				if (const UMetasoundEditorGraph* OwningGraph = ParentMember->GetOwningGraph())
+				{
+					bIsGraphEditable = OwningGraph->IsEditable();
+				}
 			}
-			DefaultCategoryBuilder->AddExternalObjectProperty(TArray<UObject*>({ DefaultFloat }), GET_MEMBER_NAME_CHECKED(UMetasoundEditorGraphMemberDefaultFloat, Range));
 
 			// add input widget properties
-			IDetailCategoryBuilder& WidgetCategoryBuilder = InDetailLayout.EditCategory("EditorOptions");
-			WidgetCategoryBuilder.AddExternalObjectProperty(TArray<UObject*>({ DefaultFloat }), GET_MEMBER_NAME_CHECKED(UMetasoundEditorGraphMemberDefaultFloat, WidgetType));
-			WidgetCategoryBuilder.AddExternalObjectProperty(TArray<UObject*>({ DefaultFloat }), GET_MEMBER_NAME_CHECKED(UMetasoundEditorGraphMemberDefaultFloat, WidgetOrientation));
-			WidgetCategoryBuilder.AddExternalObjectProperty(TArray<UObject*>({ DefaultFloat }), GET_MEMBER_NAME_CHECKED(UMetasoundEditorGraphMemberDefaultFloat, WidgetValueType));
+			if (bIsGraphEditable)
+			{
+				IDetailCategoryBuilder& WidgetCategoryBuilder = InDetailLayout.EditCategory("EditorOptions");
+				WidgetCategoryBuilder.AddExternalObjectProperty(TArray<UObject*>({ DefaultFloat }), GET_MEMBER_NAME_CHECKED(UMetasoundEditorGraphMemberDefaultFloat, WidgetType));
+				WidgetCategoryBuilder.AddExternalObjectProperty(TArray<UObject*>({ DefaultFloat }), GET_MEMBER_NAME_CHECKED(UMetasoundEditorGraphMemberDefaultFloat, WidgetOrientation));
+				WidgetCategoryBuilder.AddExternalObjectProperty(TArray<UObject*>({ DefaultFloat }), GET_MEMBER_NAME_CHECKED(UMetasoundEditorGraphMemberDefaultFloat, WidgetValueType));
+				if (DefaultFloat->WidgetValueType == EMetasoundMemberDefaultWidgetValueType::Volume)
+				{
+					WidgetCategoryBuilder.AddExternalObjectProperty(TArray<UObject*>({ DefaultFloat }), GET_MEMBER_NAME_CHECKED(UMetasoundEditorGraphMemberDefaultFloat, VolumeWidgetUseLinearOutput));
+
+				}
+			}
 		}
 
 		void FMetasoundObjectArrayLiteralCustomization::CustomizeLiteral(UMetasoundEditorGraphMemberDefaultLiteral& InLiteral, IDetailLayoutBuilder& InDetailLayout)
diff --git a/Engine/Plugins/Runtime/Metasound/Source/MetasoundEditor/Private/SMetasoundGraphNode.cpp b/Engine/Plugins/Runtime/Metasound/Source/MetasoundEditor/Private/SMetasoundGraphNode.cpp
--- a/Engine/Plugins/Runtime/Metasound/Source/MetasoundEditor/Private/SMetasoundGraphNode.cpp
+++ b/Engine/Plugins/Runtime/Metasound/Source/MetasoundEditor/Private/SMetasoundGraphNode.cpp
@@ -661,6 +661,7 @@ namespace Metasound
 								SAssignNew(InputWidget, SAudioVolumeSlider)
 									.OnValueChanged_Lambda(OnValueChangedLambda)
 									.OnValueCommitted_Lambda(OnValueCommittedLambda);
+								StaticCastSharedPtr<SAudioVolumeSlider>(InputWidget)->SetUseLinearOutput(DefaultFloat->VolumeWidgetUseLinearOutput);
 							}
 							else
 							{
@@ -739,6 +740,7 @@ namespace Metasound
 							{
 								SAssignNew(InputWidget, SAudioVolumeRadialSlider)
 									.OnValueChanged_Lambda(OnValueChangedLambda);
+								StaticCastSharedPtr<SAudioVolumeRadialSlider>(InputWidget)->SetUseLinearOutput(DefaultFloat->VolumeWidgetUseLinearOutput);
 							}
 							else
 							{
diff --git a/Engine/Plugins/Runtime/Metasound/Source/MetasoundEditor/Public/MetasoundEditorGraphMemberDefaults.h b/Engine/Plugins/Runtime/Metasound/Source/MetasoundEditor/Public/MetasoundEditorGraphMemberDefaults.h
--- a/Engine/Plugins/Runtime/Metasound/Source/MetasoundEditor/Public/MetasoundEditorGraphMemberDefaults.h
+++ b/Engine/Plugins/Runtime/Metasound/Source/MetasoundEditor/Public/MetasoundEditorGraphMemberDefaults.h
@@ -150,8 +150,8 @@ public:
 	UPROPERTY(EditAnywhere, Category = DefaultValue, meta=(EditCondition = "WidgetType == EMetasoundMemberDefaultWidget::None", EditConditionHides))
 	bool ClampDefault = false;
 
-	UPROPERTY(EditAnywhere, Category = DefaultValue, meta = (EditCondition = "ClampDefault || WidgetType != EMetasoundMemberDefaultWidget::None", EditConditionHides))
-	FVector2D Range = FVector2D(FMath::Min(0.0f, Default), FMath::Max(0.0f, Default));
+	UPROPERTY(EditAnywhere, Category = DefaultValue)
+	FVector2D Range = FVector2D(0.0f, 1.0f);
 
 	UPROPERTY(EditAnywhere, Category = Widget, meta=(DisplayName = "Widget"))
 	EMetasoundMemberDefaultWidget WidgetType = EMetasoundMemberDefaultWidget::None;
@@ -162,6 +162,10 @@ public:
 	UPROPERTY(EditAnywhere, Category = Widget, meta=(DisplayName = "Value Type", EditCondition = "WidgetType != EMetasoundMemberDefaultWidget::None", EditConditionHides))
 	EMetasoundMemberDefaultWidgetValueType WidgetValueType = EMetasoundMemberDefaultWidgetValueType::Linear;
 
+	/** If true, output linear (0 - 1) value. Otherwise, output dB value. The volume widget itself will always display the value in dB. */
+	UPROPERTY(EditAnywhere, Category = Widget, meta = (DisplayName = "Output Linear"))
+	bool VolumeWidgetUseLinearOutput = true;
+
 	static FName GetDefaultPropertyName()
 	{
 		return GET_MEMBER_NAME_CHECKED(UMetasoundEditorGraphMemberDefaultFloat, Default);
@@ -175,11 +179,12 @@ public:
 	virtual EMetasoundFrontendLiteralType GetLiteralType() const override;
 	virtual void SetFromLiteral(const FMetasoundFrontendLiteral& InLiteral) override;
 	virtual void UpdatePreviewInstance(const Metasound::FVertexName& InParameterName, TScriptInterface<IAudioParameterControllerInterface>& InParameterInterface) const override;
-	virtual void PostEditChangeProperty(FPropertyChangedEvent& PropertyChangedEvent) override;
+	virtual void PostEditChangeChainProperty(FPropertyChangedChainEvent& InPropertyChangedEvent) override;
 
 	void SetDefault(const float InDefault);
 	float GetDefault();
 	FVector2D GetRange();
+	void SetRange(const FVector2D InRange);
 };
 
 UCLASS(MinimalAPI)
