Fix various widget issues for MetaSounds

diff --git a/Engine/Plugins/Runtime/AudioWidgets/Source/AudioWidgets/Private/SAudioTextBox.cpp b/Engine/Plugins/Runtime/AudioWidgets/Source/AudioWidgets/Private/SAudioTextBox.cpp
--- a/Engine/Plugins/Runtime/AudioWidgets/Source/AudioWidgets/Private/SAudioTextBox.cpp
+++ b/Engine/Plugins/Runtime/AudioWidgets/Source/AudioWidgets/Private/SAudioTextBox.cpp
@@ -135,7 +135,7 @@ void SAudioTextBox::UpdateValueTextWidth(const FVector2D InOutputRange)
 	else
 	{
 		const float UnitsTextWidth = FontMeasureService->Measure(UnitsText->GetText(), Font).X;
-		MaxValueLabelWidth = FMath::Min(MaxValueLabelWidth, Style->BackgroundImage.ImageSize.X - UnitsTextWidth);
+		MaxValueLabelWidth = FMath::Max(MaxValueLabelWidth, Style->BackgroundImage.ImageSize.X - UnitsTextWidth);
 	}
 	ValueText->SetMinDesiredWidth(MaxValueLabelWidth);
 	if (ValueTextSlot)
diff --git a/Engine/Plugins/Runtime/AudioWidgets/Source/AudioWidgets/Public/SAudioTextBox.h b/Engine/Plugins/Runtime/AudioWidgets/Source/AudioWidgets/Public/SAudioTextBox.h
--- a/Engine/Plugins/Runtime/AudioWidgets/Source/AudioWidgets/Public/SAudioTextBox.h
+++ b/Engine/Plugins/Runtime/AudioWidgets/Source/AudioWidgets/Public/SAudioTextBox.h
@@ -28,7 +28,11 @@ struct FVariablePrecisionNumericInterface : public TDefaultNumericTypeInterface<
 		// examples: 1000, 100.1, 10.12, 1.123
 		float AbsValue = FMath::Abs(Value);
 		int32 FractionalDigits = 3;
-		if ((AbsValue / 1000.f) >= 1.f)
+
+		// special case because FastDecimalFormat::NumberToString does not parse decimal points with 0 fractional digits when a value is greater than uint64 max
+		if (AbsValue >= static_cast<float>(TNumericLimits<uint64>::Max()))
+			FractionalDigits = 1;
+		else if ((AbsValue / 1000.f) >= 1.f)
 			FractionalDigits = 0;
 		else if ((AbsValue / 100.f) >= 1.f)
 			FractionalDigits = 1;
