Fix undo/redo for MetaSound widgets and linear output for volume widgets
- Differentiate between normalized linear and audio linear value for volume output options

#rnx

diff --git a/Engine/Plugins/Runtime/AudioWidgets/Source/AudioWidgets/Private/AudioRadialSlider.cpp b/Engine/Plugins/Runtime/AudioWidgets/Source/AudioWidgets/Private/AudioRadialSlider.cpp
--- a/Engine/Plugins/Runtime/AudioWidgets/Source/AudioWidgets/Private/AudioRadialSlider.cpp
+++ b/Engine/Plugins/Runtime/AudioWidgets/Source/AudioWidgets/Private/AudioRadialSlider.cpp
@@ -61,6 +61,16 @@ UAudioRadialSlider::UAudioRadialSlider(const FObjectInitializer& ObjectInitializ
 #endif
 }
 
+float UAudioRadialSlider::GetOutputValue(const float InSliderValue)
+{
+	return MyAudioRadialSlider->GetOutputValue(InSliderValue);
+}
+
+float UAudioRadialSlider::GetSliderValue(const float OutputValue)
+{
+	return MyAudioRadialSlider->GetSliderValue(OutputValue);
+}
+
 void UAudioRadialSlider::SetWidgetLayout(EAudioRadialSliderLayout InLayout)
 {
 	WidgetLayout = InLayout;
diff --git a/Engine/Plugins/Runtime/AudioWidgets/Source/AudioWidgets/Private/AudioSlider.cpp b/Engine/Plugins/Runtime/AudioWidgets/Source/AudioWidgets/Private/AudioSlider.cpp
--- a/Engine/Plugins/Runtime/AudioWidgets/Source/AudioWidgets/Private/AudioSlider.cpp
+++ b/Engine/Plugins/Runtime/AudioWidgets/Source/AudioWidgets/Private/AudioSlider.cpp
@@ -148,14 +148,19 @@ void UAudioSliderBase::SetShowUnitsText(const bool bShowUnitsText)
 	}
 }
 
-float UAudioSliderBase::GetOutputValue(const float LinValue)
+float UAudioSliderBase::GetOutputValue(const float InSliderValue)
 {
-	return MyAudioSlider->GetOutputValue(LinValue);
+	return MyAudioSlider->GetOutputValue(InSliderValue);
 }
 
 float UAudioSliderBase::GetLinValue(const float OutputValue)
 {
-	return MyAudioSlider->GetLinValue(OutputValue);
+	return GetSliderValue(OutputValue);
+}
+
+float UAudioSliderBase::GetSliderValue(const float OutputValue)
+{
+	return MyAudioSlider->GetSliderValue(OutputValue);
 }
 
 void UAudioSliderBase::SetSliderBackgroundColor(FLinearColor InValue)
diff --git a/Engine/Plugins/Runtime/AudioWidgets/Source/AudioWidgets/Private/SAudioRadialSlider.cpp b/Engine/Plugins/Runtime/AudioWidgets/Source/AudioWidgets/Private/SAudioRadialSlider.cpp
--- a/Engine/Plugins/Runtime/AudioWidgets/Source/AudioWidgets/Private/SAudioRadialSlider.cpp
+++ b/Engine/Plugins/Runtime/AudioWidgets/Source/AudioWidgets/Private/SAudioRadialSlider.cpp
@@ -6,7 +6,7 @@
 #include "Widgets/SRadialSlider.h"
 #include "Styling/SlateStyleRegistry.h"
 
-const FVector2D SAudioRadialSlider::LinearRange = FVector2D(0.0f, 1.0f);
+const FVector2D SAudioRadialSlider::NormalizedLinearSliderRange = FVector2D(0.0f, 1.0f);
 
 SAudioRadialSlider::SAudioRadialSlider()
 {
@@ -16,7 +16,9 @@ void SAudioRadialSlider::Construct(const SAudioRadialSlider::FArguments& InArgs)
 {
 	Style = InArgs._Style;
 	OnValueChanged = InArgs._OnValueChanged;
-	Value = InArgs._Value;
+	OnMouseCaptureBegin = InArgs._OnMouseCaptureBegin;
+	OnMouseCaptureEnd = InArgs._OnMouseCaptureEnd;
+	SliderValue = InArgs._SliderValue;
 	CenterBackgroundColor = InArgs._CenterBackgroundColor;
 	SliderProgressColor = InArgs._SliderProgressColor;
 	SliderBarColor = InArgs._SliderBarColor;
@@ -41,20 +43,26 @@ void SAudioRadialSlider::Construct(const SAudioRadialSlider::FArguments& InArgs)
 		.OnValueTextCommitted_Lambda([this](const FText& Text, ETextCommit::Type CommitType)
 		{
 			const float OutputValue = FCString::Atof(*Text.ToString());
-			const float LinValue = GetLinValueForText(OutputValue);
-			Value.Set(LinValue);
-			RadialSlider->SetValue(LinValue);
-			OnValueChanged.ExecuteIfBound(LinValue);
+			const float NewSliderValue = GetSliderValueForText(OutputValue);
+			if (!FMath::IsNearlyEqual(NewSliderValue, SliderValue.Get()))
+			{
+				SliderValue.Set(NewSliderValue);
+				RadialSlider->SetValue(NewSliderValue);
+				OnValueChanged.ExecuteIfBound(NewSliderValue);
+				OnMouseCaptureEnd.ExecuteIfBound();
+			}
 		});
 
 	SAssignNew(RadialSlider, SRadialSlider)
-		.OnValueChanged_Lambda([this](float InLinValue)
+		.OnValueChanged_Lambda([this](float InSliderValue)
 		{
-			Value.Set(InLinValue);
-			OnValueChanged.ExecuteIfBound(InLinValue);
-			const float OutputValue = GetOutputValueForText(InLinValue);
+			SliderValue.Set(InSliderValue);
+			OnValueChanged.ExecuteIfBound(InSliderValue);
+			const float OutputValue = GetOutputValueForText(InSliderValue);
 			Label->SetValueText(OutputValue);
 		})
+		.OnMouseCaptureBegin(OnMouseCaptureBegin)
+		.OnMouseCaptureEnd(OnMouseCaptureEnd)
 		.UseVerticalDrag(true)
 		.ShowSliderHand(true)
 		.ShowSliderHandle(false);
@@ -184,22 +192,22 @@ TSharedRef<SWidgetSwitcher> SAudioRadialSlider::CreateLayoutWidgetSwitcher()
 	return LayoutWidgetSwitcher.ToSharedRef();
 }
 
-void SAudioRadialSlider::SetValue(float LinValue)
+void SAudioRadialSlider::SetSliderValue(float InSliderValue)
 {
-	Value.Set(LinValue);
-	const float OutputValueText = GetOutputValueForText(LinValue);
+	SliderValue.Set(InSliderValue);
+	const float OutputValueText = GetOutputValueForText(InSliderValue);
 	Label->SetValueText(OutputValueText);
-	RadialSlider->SetValue(LinValue);
+	RadialSlider->SetValue(InSliderValue);
 }
 
-const float SAudioRadialSlider::GetLinValue(const float OutputValue)
+const float SAudioRadialSlider::GetSliderValue(const float OutputValue)
 {
-	return FMath::GetMappedRangeValueClamped(OutputRange, LinearRange, OutputValue);
+	return FMath::GetMappedRangeValueClamped(OutputRange, NormalizedLinearSliderRange, OutputValue);
 }
 
-const float SAudioRadialSlider::GetOutputValue(const float LinValue)
+const float SAudioRadialSlider::GetOutputValue(const float InSliderValue)
 {
-	return FMath::GetMappedRangeValueClamped(LinearRange, OutputRange, LinValue);
+	return FMath::GetMappedRangeValueClamped(NormalizedLinearSliderRange, OutputRange, InSliderValue);
 }
 
 void SAudioRadialSlider::SetOutputRange(const FVector2D Range)
@@ -208,22 +216,22 @@ void SAudioRadialSlider::SetOutputRange(const FVector2D Range)
 	// if Range.Y < Range.X, set Range.X to Range.Y
 	OutputRange.X = FMath::Min(Range.X, Range.Y);
 
-	const float OutputValue = GetOutputValue(Value.Get());
+	const float OutputValue = GetOutputValue(SliderValue.Get());
 	const float ClampedOutputValue = FMath::Clamp(OutputValue, OutputRange.X, OutputRange.Y);
-	const float LinValue = GetLinValue(ClampedOutputValue);
-	SetValue(LinValue);
+	const float ClampedSliderValue = GetSliderValue(ClampedOutputValue);
+	SetSliderValue(ClampedSliderValue);
 
 	Label->UpdateValueTextWidth(OutputRange);
 }
 
-const float SAudioRadialSlider::GetOutputValueForText(const float LinValue)
+const float SAudioRadialSlider::GetOutputValueForText(const float InSliderValue)
 {
-	return GetOutputValue(LinValue);
+	return GetOutputValue(InSliderValue);
 }
 
-const float SAudioRadialSlider::GetLinValueForText(const float OutputValue)
+const float SAudioRadialSlider::GetSliderValueForText(const float OutputValue)
 {
-	return GetLinValue(OutputValue);
+	return GetSliderValue(OutputValue);
 }
 
 void SAudioRadialSlider::SetLabelBackgroundColor(FSlateColor InColor)
@@ -273,63 +281,67 @@ SAudioVolumeRadialSlider::SAudioVolumeRadialSlider()
 void SAudioVolumeRadialSlider::Construct(const SAudioRadialSlider::FArguments& InArgs)
 {
 	SAudioRadialSlider::Construct(InArgs);
-	
-	SAudioRadialSlider::SetOutputRange(FVector2D(-100.0f, 0.0f));
+
+	SAudioRadialSlider::SetOutputRange(FVector2D(SAudioVolumeRadialSlider::MinDbValue, 0.0f));
 	Label->SetUnitsText(FText::FromString("dB"));
 }
 
-const float SAudioVolumeRadialSlider::GetDbValueFromLin(const float LinValue)
+const float SAudioVolumeRadialSlider::GetDbValueFromSliderValue(const float InSliderValue)  
 {
 	// convert from linear 0-1 space to decibel OutputRange that has been converted to linear 
 	const FVector2D LinearSliderRange = FVector2D(Audio::ConvertToLinear(OutputRange.X), Audio::ConvertToLinear(OutputRange.Y));
-	const float LinearSliderValue = FMath::GetMappedRangeValueClamped(LinearRange, LinearSliderRange, LinValue);
+	const float LinearSliderValue = FMath::GetMappedRangeValueClamped(NormalizedLinearSliderRange, LinearSliderRange, InSliderValue);
 	// convert from linear to decibels 
 	float OutputValue = Audio::ConvertToDecibels(LinearSliderValue);
 	return FMath::Clamp(OutputValue, OutputRange.X, OutputRange.Y);
 }
 
-const float SAudioVolumeRadialSlider::GetLinValueFromDb(const float OutputValue)
+const float SAudioVolumeRadialSlider::GetSliderValueFromDb(const float OutputValue) 
 {
 	float ClampedValue = FMath::Clamp(OutputValue, OutputRange.X, OutputRange.Y);
 	// convert from decibels to linear
 	float LinearSliderValue = Audio::ConvertToLinear(ClampedValue);
 	// convert from decibel OutputRange that has been converted to linear to linear 0-1 space 
 	const FVector2D LinearSliderRange = FVector2D(Audio::ConvertToLinear(OutputRange.X), Audio::ConvertToLinear(OutputRange.Y));
-	return FMath::GetMappedRangeValueClamped(LinearSliderRange, LinearRange, LinearSliderValue);
+	return FMath::GetMappedRangeValueClamped(LinearSliderRange, NormalizedLinearSliderRange, LinearSliderValue);
 }
 
-const float SAudioVolumeRadialSlider::GetOutputValue(const float LinValue)
+const float SAudioVolumeRadialSlider::GetOutputValue(const float InSliderValue) 
 {
 	if (bUseLinearOutput)
 	{
-		return FMath::Clamp(LinValue, LinearRange.X, LinearRange.Y);
+		// Return linear given normalized linear 
+		const FVector2D LinearSliderRange = FVector2D(Audio::ConvertToLinear(OutputRange.X), Audio::ConvertToLinear(OutputRange.Y));
+		return FMath::GetMappedRangeValueClamped(NormalizedLinearSliderRange, LinearSliderRange, InSliderValue);
 	}
 	else
 	{
-		return GetDbValueFromLin(LinValue);
+		return GetDbValueFromSliderValue(InSliderValue);
 	}
 }
 
-const float SAudioVolumeRadialSlider::GetLinValue(const float OutputValue)
+const float SAudioVolumeRadialSlider::GetSliderValue(const float OutputValue)
 {
 	if (bUseLinearOutput)
 	{
-		return FMath::Clamp(OutputValue, LinearRange.X, LinearRange.Y);
+		// Convert from linear to normalized linear 
+		const FVector2D LinearSliderRange = FVector2D(Audio::ConvertToLinear(OutputRange.X), Audio::ConvertToLinear(OutputRange.Y));
+		return FMath::GetMappedRangeValueClamped(LinearSliderRange, NormalizedLinearSliderRange, OutputValue);
 	}
 	else
 	{
-		return GetLinValueFromDb(OutputValue);
+		return GetSliderValueFromDb(OutputValue);
 	}
 }
 
-const float SAudioVolumeRadialSlider::GetOutputValueForText(const float LinValue)
+const float SAudioVolumeRadialSlider::GetOutputValueForText(const float InSliderValue)
 {
-	return GetDbValueFromLin(LinValue);
+	return GetDbValueFromSliderValue(InSliderValue);
 }
 
-const float SAudioVolumeRadialSlider::GetLinValueForText(const float OutputValue)
+const float SAudioVolumeRadialSlider::GetSliderValueForText(const float OutputValue)
 {
-	return GetLinValueFromDb(OutputValue);
+	return GetSliderValueFromDb(OutputValue);
 }
 
 void SAudioVolumeRadialSlider::SetUseLinearOutput(bool InUseLinearOutput)
@@ -337,12 +349,21 @@ void SAudioVolumeRadialSlider::SetUseLinearOutput(bool InUseLinearOutput)
 	bUseLinearOutput = InUseLinearOutput;
 }
 
-void SAudioVolumeRadialSlider::SetOutputRange(const FVector2D Range)
+void SAudioVolumeRadialSlider::SetOutputRange(const FVector2D Range) 
 {
-	// if using linear output, output range cannot be changed 
-	if (!bUseLinearOutput)
+	// For volume slider, OutputRange is always in dB
+	if (!(Range - OutputRange).IsNearlyZero())
 	{
-		SAudioRadialSlider::SetOutputRange(FVector2D(FMath::Max(MinDbValue, Range.X), FMath::Min(MaxDbValue, Range.Y)));
+		if (bUseLinearOutput)
+		{
+			// If using linear output, assume given range is linear (not normalized though) 
+			FVector2D RangeInDecibels = FVector2D(Audio::ConvertToDecibels(Range.X), Audio::ConvertToDecibels(Range.Y));
+			SAudioRadialSlider::SetOutputRange(FVector2D(FMath::Max(MinDbValue, RangeInDecibels.X), FMath::Min(MaxDbValue, RangeInDecibels.Y)));
+		}
+		else
+		{
+			SAudioRadialSlider::SetOutputRange(FVector2D(FMath::Max(MinDbValue, Range.X), FMath::Min(MaxDbValue, Range.Y)));
+		}
 	}
 }
 
@@ -359,21 +380,21 @@ void SAudioFrequencyRadialSlider::Construct(const SAudioRadialSlider::FArguments
 	Label->SetUnitsText(FText::FromString("Hz"));
 }
 
-const float SAudioFrequencyRadialSlider::GetOutputValue(const float LinValue)
+const float SAudioFrequencyRadialSlider::GetOutputValue(const float InSliderValue)
 {
-	return Audio::GetLogFrequencyClamped(LinValue, LinearRange, OutputRange);
+	return Audio::GetLogFrequencyClamped(InSliderValue, NormalizedLinearSliderRange, OutputRange);
 }
 
-const float SAudioFrequencyRadialSlider::GetLinValue(const float OutputValue)
+const float SAudioFrequencyRadialSlider::GetSliderValue(const float OutputValue)
 {
 	// edge case to avoid audio function returning negative value
 	if (FMath::IsNearlyEqual(OutputValue, OutputRange.X))
 	{
-		return LinearRange.X;
+		return NormalizedLinearSliderRange.X;
 	}
 	if (FMath::IsNearlyEqual(OutputValue, OutputRange.Y))
 	{
-		return LinearRange.Y;
+		return NormalizedLinearSliderRange.Y;
 	}
-	return Audio::GetLinearFrequencyClamped(OutputValue, LinearRange, OutputRange);
+	return Audio::GetLinearFrequencyClamped(OutputValue, NormalizedLinearSliderRange, OutputRange);
 }
diff --git a/Engine/Plugins/Runtime/AudioWidgets/Source/AudioWidgets/Private/SAudioSlider.cpp b/Engine/Plugins/Runtime/AudioWidgets/Source/AudioWidgets/Private/SAudioSlider.cpp
--- a/Engine/Plugins/Runtime/AudioWidgets/Source/AudioWidgets/Private/SAudioSlider.cpp
+++ b/Engine/Plugins/Runtime/AudioWidgets/Source/AudioWidgets/Private/SAudioSlider.cpp
@@ -13,7 +13,7 @@
 #include "Styling/SlateStyleRegistry.h"
 
 // SAudioSliderBase
-const FVector2D SAudioSliderBase::LinearRange = FVector2D(0.0f, 1.0f);
+const FVector2D SAudioSliderBase::NormalizedLinearSliderRange = FVector2D(0.0f, 1.0f);
 
 SAudioSliderBase::SAudioSliderBase()
 {
@@ -24,7 +24,7 @@ void SAudioSliderBase::Construct(const SAudioSliderBase::FArguments& InArgs)
 	Style = InArgs._Style;
 	OnValueChanged = InArgs._OnValueChanged;
 	OnValueCommitted = InArgs._OnValueCommitted;
-	ValueAttribute = InArgs._Value;
+	SliderValueAttribute = InArgs._SliderValue;
 	SliderBackgroundColor = InArgs._SliderBackgroundColor;
 	SliderBarColor = InArgs._SliderBarColor;
 	SliderThumbColor = InArgs._SliderThumbColor;
@@ -51,20 +51,20 @@ void SAudioSliderBase::Construct(const SAudioSliderBase::FArguments& InArgs)
 
 	// Underlying slider widget
 	SAssignNew(Slider, SSlider)
-		.Value(ValueAttribute.Get())
+		.Value(SliderValueAttribute.Get())
 		.Style(&Style->SliderStyle)
 		.Orientation(Orientation.Get())
 		.IndentHandle(false)
 		.OnValueChanged_Lambda([this](float Value)
 		{
-			ValueAttribute.Set(Value);
+			SliderValueAttribute.Set(Value);
 			OnValueChanged.ExecuteIfBound(Value);
 			const float OutputValue = GetOutputValueForText(Value);
 			Label->SetValueText(OutputValue);
 		})
 		.OnMouseCaptureEnd_Lambda([this]()
 		{
-			OnValueCommitted.ExecuteIfBound(ValueAttribute.Get());
+			OnValueCommitted.ExecuteIfBound(SliderValueAttribute.Get());
 		});
 
 	// Text label
@@ -73,11 +73,14 @@ void SAudioSliderBase::Construct(const SAudioSliderBase::FArguments& InArgs)
 		.OnValueTextCommitted_Lambda([this](const FText& Text, ETextCommit::Type CommitType)
 		{
 			const float OutputValue = FCString::Atof(*Text.ToString());
-			const float LinValue = GetLinValueForText(OutputValue);
-			ValueAttribute.Set(LinValue);
-			Slider->SetValue(LinValue);
-			OnValueChanged.ExecuteIfBound(LinValue);
-			OnValueCommitted.ExecuteIfBound(LinValue);
+			const float NewSliderValue = GetSliderValueForText(OutputValue);
+			if (!FMath::IsNearlyEqual(NewSliderValue, SliderValueAttribute.Get()))
+			{
+				SliderValueAttribute.Set(NewSliderValue);
+				Slider->SetValue(NewSliderValue);
+				OnValueChanged.ExecuteIfBound(NewSliderValue);
+				OnValueCommitted.ExecuteIfBound(NewSliderValue);
+			}
 		});
 
 	ChildSlot
@@ -86,12 +89,12 @@ void SAudioSliderBase::Construct(const SAudioSliderBase::FArguments& InArgs)
 	];
 }
 
-void SAudioSliderBase::SetValue(float LinValue)
+void SAudioSliderBase::SetSliderValue(float InSliderValue)
 {
-	ValueAttribute.Set(LinValue);
-	const float OutputValueForText = GetOutputValueForText(LinValue);
+	SliderValueAttribute.Set(InSliderValue);
+	const float OutputValueForText = GetOutputValueForText(InSliderValue);
 	Label->SetValueText(OutputValueForText);
-	Slider->SetValue(LinValue);
+	Slider->SetValue(InSliderValue);
 }
 
 void SAudioSliderBase::SetOrientation(EOrientation InOrientation)
@@ -136,22 +139,22 @@ void SAudioSliderBase::SetDesiredSizeOverride(const FVector2D Size)
 	SetAttribute(DesiredSizeOverride, TAttribute<TOptional<FVector2D>>(Size), EInvalidateWidgetReason::Layout);
 }
 
-const float SAudioSliderBase::GetOutputValue(const float LinValue)
+const float SAudioSliderBase::GetOutputValue(const float InSliderValue)
 {
-	return FMath::Clamp(LinValue, OutputRange.X, OutputRange.Y);
+	return FMath::Clamp(InSliderValue, OutputRange.X, OutputRange.Y);
 }
 
-const float SAudioSliderBase::GetOutputValueForText(const float LinValue)
+const float SAudioSliderBase::GetOutputValueForText(const float InSliderValue)
 {
-	return GetOutputValue(LinValue);
+	return GetOutputValue(InSliderValue);
 }
 
-const float SAudioSliderBase::GetLinValueForText(const float OutputValue)
+const float SAudioSliderBase::GetSliderValueForText(const float OutputValue)
 {
-	return GetLinValue(OutputValue);
+	return GetSliderValue(OutputValue);
 }
 
-const float SAudioSliderBase::GetLinValue(const float OutputValue)
+const float SAudioSliderBase::GetSliderValue(const float OutputValue)
 {
 	return OutputValue;
 }
@@ -188,10 +191,10 @@ void SAudioSliderBase::SetOutputRange(const FVector2D Range)
 	// if Range.Y < Range.X, set Range.X to Range.Y
 	OutputRange.X = FMath::Min(Range.X, Range.Y);
 
-	const float OutputValue = GetOutputValue(ValueAttribute.Get());
+	const float OutputValue = GetOutputValue(SliderValueAttribute.Get());
 	const float ClampedOutputValue = FMath::Clamp(OutputValue, OutputRange.X, OutputRange.Y);
-	const float LinValue = GetLinValue(ClampedOutputValue);
-	SetValue(LinValue);
+	const float ClampedSliderValue = GetSliderValue(ClampedOutputValue);
+	SetSliderValue(ClampedSliderValue);
 
 	Label->UpdateValueTextWidth(OutputRange);
 }
@@ -347,7 +350,7 @@ void SAudioSlider::Construct(const SAudioSliderBase::FArguments& InArgs)
 void SAudioSlider::SetLinToOutputCurve(const TWeakObjectPtr<const UCurveFloat> InLinToOutputCurve)
 {
 	LinToOutputCurve = InLinToOutputCurve;
-	Label->SetValueText(GetOutputValueForText(ValueAttribute.Get()));
+	Label->SetValueText(GetOutputValueForText(SliderValueAttribute.Get()));
 }
 
 void SAudioSlider::SetOutputToLinCurve(const TWeakObjectPtr<const UCurveFloat> InOutputToLinCurve)
@@ -365,24 +368,24 @@ const TWeakObjectPtr<const UCurveFloat> SAudioSlider::GetLinToOutputCurve()
 	return LinToOutputCurve;
 }
 
-const float SAudioSlider::GetOutputValue(const float LinValue)
+const float SAudioSlider::GetOutputValue(const float InSliderValue)
 {
 	if (LinToOutputCurve.IsValid())
 	{
-		const float CurveOutputValue = LinToOutputCurve->GetFloatValue(LinValue);
+		const float CurveOutputValue = LinToOutputCurve->GetFloatValue(InSliderValue);
 		return FMath::Clamp(CurveOutputValue, OutputRange.X, OutputRange.Y);
 	}
-	return FMath::GetMappedRangeValueClamped(LinearRange, OutputRange, LinValue);
+	return FMath::GetMappedRangeValueClamped(NormalizedLinearSliderRange, OutputRange, InSliderValue);
 }
 
-const float SAudioSlider::GetLinValue(const float OutputValue)
+const float SAudioSlider::GetSliderValue(const float OutputValue)
 {
 	if (OutputToLinCurve.IsValid())
 	{
-		const float CurveLinValue = OutputToLinCurve->GetFloatValue(OutputValue);
-		return FMath::Clamp(CurveLinValue, OutputRange.X, OutputRange.Y);
+		const float CurveSliderValue = OutputToLinCurve->GetFloatValue(OutputValue);
+		return FMath::Clamp(CurveSliderValue, OutputRange.X, OutputRange.Y);
 	}
-	return FMath::GetMappedRangeValueClamped(OutputRange, LinearRange, OutputValue);
+	return FMath::GetMappedRangeValueClamped(OutputRange, NormalizedLinearSliderRange, OutputValue);
 }
 
 // SAudioVolumeSlider
@@ -399,58 +402,62 @@ void SAudioVolumeSlider::Construct(const SAudioSlider::FArguments& InArgs)
 	Label->SetUnitsText(FText::FromString("dB"));
 }
 
-const float SAudioVolumeSlider::GetDbValueFromLin(const float LinValue)
+const float SAudioVolumeSlider::GetDbValueFromSliderValue(const float InSliderValue)
 {
 	// convert from linear 0-1 space to decibel OutputRange that has been converted to linear 
 	const FVector2D LinearSliderRange = FVector2D(Audio::ConvertToLinear(OutputRange.X), Audio::ConvertToLinear(OutputRange.Y));
-	const float LinearSliderValue = FMath::GetMappedRangeValueClamped(LinearRange, LinearSliderRange, LinValue);
+	const float LinearSliderValue = FMath::GetMappedRangeValueClamped(NormalizedLinearSliderRange, LinearSliderRange, InSliderValue);
 	// convert from linear to decibels 
 	float OutputValue = Audio::ConvertToDecibels(LinearSliderValue);
 	return FMath::Clamp(OutputValue, OutputRange.X, OutputRange.Y);
 }
 
-const float SAudioVolumeSlider::GetLinValueFromDb(const float DbValue)
+const float SAudioVolumeSlider::GetSliderValueFromDb(const float DbValue)
 {
 	float ClampedValue = FMath::Clamp(DbValue, OutputRange.X, OutputRange.Y);
 	// convert from decibels to linear
 	float LinearSliderValue = Audio::ConvertToLinear(ClampedValue);
 	// convert from decibel OutputRange that has been converted to linear to linear 0-1 space 
 	const FVector2D LinearSliderRange = FVector2D(Audio::ConvertToLinear(OutputRange.X), Audio::ConvertToLinear(OutputRange.Y));
-	return FMath::GetMappedRangeValueClamped(LinearSliderRange, LinearRange, LinearSliderValue);
+	return FMath::GetMappedRangeValueClamped(LinearSliderRange, NormalizedLinearSliderRange, LinearSliderValue);
 }
 
-const float SAudioVolumeSlider::GetOutputValue(const float LinValue)
+const float SAudioVolumeSlider::GetOutputValue(const float InSliderValue)
 {
 	if (bUseLinearOutput)
 	{
-		return FMath::Clamp(LinValue, LinearRange.X, LinearRange.Y);
+		// Return linear given normalized linear 
+		const FVector2D LinearSliderRange = FVector2D(Audio::ConvertToLinear(OutputRange.X), Audio::ConvertToLinear(OutputRange.Y));
+		return FMath::GetMappedRangeValueClamped(NormalizedLinearSliderRange, LinearSliderRange, InSliderValue);
 	}
 	else
 	{
-		return GetDbValueFromLin(LinValue);
+		return GetDbValueFromSliderValue(InSliderValue);
 	}
 }
 
-const float SAudioVolumeSlider::GetLinValue(const float OutputValue)
+const float SAudioVolumeSlider::GetSliderValue(const float OutputValue)
 {
 	if (bUseLinearOutput)
 	{
-		return FMath::Clamp(OutputValue, LinearRange.X, LinearRange.Y);
+		// Convert from linear to normalized linear 
+		const FVector2D LinearSliderRange = FVector2D(Audio::ConvertToLinear(OutputRange.X), Audio::ConvertToLinear(OutputRange.Y));
+		return FMath::GetMappedRangeValueClamped(LinearSliderRange, NormalizedLinearSliderRange, OutputValue);
 	}
 	else
 	{
-		return GetLinValueFromDb(OutputValue);
+		return GetSliderValueFromDb(OutputValue);
 	}
 }
 
-const float SAudioVolumeSlider::GetOutputValueForText(const float LinValue)
+const float SAudioVolumeSlider::GetOutputValueForText(const float InSliderValue)
 {
-	return GetDbValueFromLin(LinValue);
+	return GetDbValueFromSliderValue(InSliderValue);
 }
 
-const float SAudioVolumeSlider::GetLinValueForText(const float OutputValue)
+const float SAudioVolumeSlider::GetSliderValueForText(const float OutputValue)
 {
-	return GetLinValueFromDb(OutputValue);
+	return GetSliderValueFromDb(OutputValue);
 }
 
 void SAudioVolumeSlider::SetUseLinearOutput(bool InUseLinearOutput)
@@ -460,10 +467,19 @@ void SAudioVolumeSlider::SetUseLinearOutput(bool InUseLinearOutput)
 
 void SAudioVolumeSlider::SetOutputRange(const FVector2D Range)
 {
-	// if using linear output, output range cannot be changed 
-	if (!bUseLinearOutput)
+	// For volume slider, OutputRange is always in dB
+	if (!(Range - OutputRange).IsNearlyZero())
 	{
-		SAudioSliderBase::SetOutputRange(FVector2D(FMath::Max(MinDbValue, Range.X), FMath::Min(MaxDbValue, Range.Y)));
+		if (bUseLinearOutput)
+		{
+			// If using linear output, assume given range is linear (not normalized though) 
+			FVector2D RangeInDecibels = FVector2D(Audio::ConvertToDecibels(Range.X), Audio::ConvertToDecibels(Range.Y));
+			SAudioSliderBase::SetOutputRange(FVector2D(FMath::Max(MinDbValue, RangeInDecibels.X), FMath::Min(MaxDbValue, RangeInDecibels.Y)));
+		}
+		else
+		{
+			SAudioSliderBase::SetOutputRange(FVector2D(FMath::Max(MinDbValue, Range.X), FMath::Min(MaxDbValue, Range.Y)));
+		}
 	}
 }
 
@@ -479,21 +495,21 @@ void SAudioFrequencySlider::Construct(const SAudioSlider::FArguments& InArgs)
 	Label->SetUnitsText(FText::FromString("Hz"));
 }
 
-const float SAudioFrequencySlider::GetOutputValue(const float LinValue)
+const float SAudioFrequencySlider::GetOutputValue(const float InSliderValue)
 {
-	return Audio::GetLogFrequencyClamped(LinValue, LinearRange, OutputRange);
+	return Audio::GetLogFrequencyClamped(InSliderValue, NormalizedLinearSliderRange, OutputRange);
 }
 
-const float SAudioFrequencySlider::GetLinValue(const float OutputValue)
+const float SAudioFrequencySlider::GetSliderValue(const float OutputValue)
 {
 	// edge case to avoid audio function returning negative value
 	if (FMath::IsNearlyEqual(OutputValue, OutputRange.X))
 	{
-		return LinearRange.X;
+		return NormalizedLinearSliderRange.X;
 	}
 	if (FMath::IsNearlyEqual(OutputValue, OutputRange.Y))
 	{
-		return LinearRange.Y;
+		return NormalizedLinearSliderRange.Y;
 	}
-	return Audio::GetLinearFrequencyClamped(OutputValue, LinearRange, OutputRange);
+	return Audio::GetLinearFrequencyClamped(OutputValue, NormalizedLinearSliderRange, OutputRange);
 }
diff --git a/Engine/Plugins/Runtime/AudioWidgets/Source/AudioWidgets/Public/AudioRadialSlider.h b/Engine/Plugins/Runtime/AudioWidgets/Source/AudioWidgets/Public/AudioRadialSlider.h
--- a/Engine/Plugins/Runtime/AudioWidgets/Source/AudioWidgets/Public/AudioRadialSlider.h
+++ b/Engine/Plugins/Runtime/AudioWidgets/Source/AudioWidgets/Public/AudioRadialSlider.h
@@ -21,7 +21,7 @@ class AUDIOWIDGETS_API UAudioRadialSlider : public UWidget
 	GENERATED_UCLASS_BODY()
 
 public:
-	/** The linear value. */
+	/** The normalized linear (0 - 1) slider value position. */
 	UPROPERTY(EditAnywhere, Category = Appearance, meta = (UIMin = "0", UIMax = "1"))
 	float Value;
 
@@ -86,6 +86,14 @@ public:
 	FOnAudioRadialSliderValueChangedEvent OnValueChanged;
 
 public: 
+	/** Get output value from normalized linear (0 - 1) based on internal lin to output mapping. */
+	UFUNCTION(BlueprintCallable, Category = "Behavior")
+	float GetOutputValue(const float InSliderValue);
+
+	/** Get normalized linear (0 - 1) slider value from output based on internal lin to output mapping. */
+	UFUNCTION(BlueprintCallable, Category = "Behavior")
+	float GetSliderValue(const float OutputValue);
+
 	/** Sets the widget layout */
 	UFUNCTION(BlueprintCallable, Category = "Appearance")
 	void SetWidgetLayout(EAudioRadialSliderLayout InLayout);
diff --git a/Engine/Plugins/Runtime/AudioWidgets/Source/AudioWidgets/Public/AudioSlider.h b/Engine/Plugins/Runtime/AudioWidgets/Source/AudioWidgets/Public/AudioSlider.h
--- a/Engine/Plugins/Runtime/AudioWidgets/Source/AudioWidgets/Public/AudioSlider.h
+++ b/Engine/Plugins/Runtime/AudioWidgets/Source/AudioWidgets/Public/AudioSlider.h
@@ -25,7 +25,7 @@ class AUDIOWIDGETS_API UAudioSliderBase: public UWidget
 	GENERATED_UCLASS_BODY()
 
 public:
-	/** The linear value. */
+	/** The normalized linear (0 - 1) slider value. */
 	UPROPERTY(EditAnywhere, Category = Appearance, meta = (UIMin = "0", UIMax = "1"))
 	float Value;
 
@@ -94,14 +94,18 @@ public:
 	FGetLinearColor WidgetBackgroundColorDelegate;
 
 public:
-	/** Get output value from linear based on internal lin to output mapping. */
+	/** Get output value from normalized linear (0 - 1) based on internal lin to output mapping. */
 	UFUNCTION(BlueprintCallable, Category = "Behavior")
-	float GetOutputValue(const float LinValue);
+	float GetOutputValue(const float InSliderValue);
 
-	/** Get linear value from output based on internal lin to output mapping. */
-	UFUNCTION(BlueprintCallable, Category = "Behavior")
+	/** Get normalized linear (0 - 1) value from output based on internal lin to output mapping. */
+	UFUNCTION(BlueprintCallable, Category = "Behavior", meta=(DeprecatedFunction, DeprecationMessage="5.1 - GetLinValue is deprecated, please use GetSliderValue instead."))
 	float GetLinValue(const float OutputValue);
 
+	/** Get normalized linear (0 - 1) slider value from output based on internal lin to output mapping. */
+	UFUNCTION(BlueprintCallable, Category = "Behavior")
+	float GetSliderValue(const float OutputValue);
+
 	/** Sets the label background color */
 	UFUNCTION(BlueprintCallable, Category = "Appearance")
 	void SetTextLabelBackgroundColor(FSlateColor InColor);
diff --git a/Engine/Plugins/Runtime/AudioWidgets/Source/AudioWidgets/Public/SAudioInputWidget.h b/Engine/Plugins/Runtime/AudioWidgets/Source/AudioWidgets/Public/SAudioInputWidget.h
--- a/Engine/Plugins/Runtime/AudioWidgets/Source/AudioWidgets/Public/SAudioInputWidget.h
+++ b/Engine/Plugins/Runtime/AudioWidgets/Source/AudioWidgets/Public/SAudioInputWidget.h
@@ -14,13 +14,13 @@
 	: public SCompoundWidget
 {
 public:
-	virtual const float GetOutputValue(const float LinValue) = 0;
-	virtual const float GetLinValue(const float OutputValue) = 0;
+	virtual const float GetOutputValue(const float InSliderValue) = 0;
+	virtual const float GetSliderValue(const float OutputValue) = 0;
 	
 	/**
 	 * Set the slider's linear (0-1 normalized) value. 
 	 */
-	virtual void SetValue(float LinValue) = 0;
+	virtual void SetSliderValue(float InSliderValue) = 0;
 	virtual void SetOutputRange(const FVector2D Range) = 0;
 	
 	virtual void SetLabelBackgroundColor(FSlateColor InColor) = 0;
diff --git a/Engine/Plugins/Runtime/AudioWidgets/Source/AudioWidgets/Public/SAudioRadialSlider.h b/Engine/Plugins/Runtime/AudioWidgets/Source/AudioWidgets/Public/SAudioRadialSlider.h
--- a/Engine/Plugins/Runtime/AudioWidgets/Source/AudioWidgets/Public/SAudioRadialSlider.h
+++ b/Engine/Plugins/Runtime/AudioWidgets/Source/AudioWidgets/Public/SAudioRadialSlider.h
@@ -43,7 +43,7 @@ class AUDIOWIDGETS_API SAudioRadialSlider
 public:
 	SLATE_BEGIN_ARGS(SAudioRadialSlider)
 	{
-		_Value = 0.0f;
+		_SliderValue = 0.0f;
 		_WidgetLayout = EAudioRadialSliderLayout::Layout_LabelBottom;
 		_HandStartEndRatio = FVector2D(0.0f, 1.0f);
 
@@ -60,8 +60,8 @@ public:
 	/** The style used to draw the audio radial slider. */
 	SLATE_STYLE_ARGUMENT(FAudioRadialSliderStyle, Style)
 
-	/** A value representing the audio slider value. */
-	SLATE_ATTRIBUTE(float, Value)
+	/** A value representing the normalized linear (0 - 1) audio slider value position. */
+	SLATE_ATTRIBUTE(float, SliderValue)
 
 	/** The widget layout. */
 	SLATE_ATTRIBUTE(EAudioRadialSliderLayout, WidgetLayout)
@@ -86,7 +86,13 @@ public:
 
 	/** Called when the value is changed by slider or typing */
 	SLATE_EVENT(FOnFloatValueChanged, OnValueChanged)
-		
+
+	/** Invoked when the mouse is pressed and a capture begins. */
+	SLATE_EVENT(FSimpleDelegate, OnMouseCaptureBegin)
+
+	/** Invoked when the mouse is released and a capture ends. */
+	SLATE_EVENT(FSimpleDelegate, OnMouseCaptureEnd)
+
 	SLATE_END_ARGS()
 
 	SAudioRadialSlider();
@@ -94,6 +100,12 @@ public:
 
 	// Holds a delegate that is executed when the slider's value changed.
 	FOnFloatValueChanged OnValueChanged;
+
+	// Holds a delegate that is executed when the mouse is pressed and a capture begins.
+	FSimpleDelegate OnMouseCaptureBegin;
+
+	// Holds a delegate that is executed when the mouse is let up and a capture ends.
+	FSimpleDelegate OnMouseCaptureEnd;
 	
 	void SetCenterBackgroundColor(FSlateColor InColor);
 	void SetSliderProgressColor(FSlateColor InColor);
@@ -105,11 +117,11 @@ public:
 	void SetDesiredSizeOverride(const FVector2D DesiredSize);
 
 	virtual void Construct(const SAudioRadialSlider::FArguments& InArgs);
-	void SetValue(float LinValue);
-	virtual const float GetOutputValue(const float LinValue);
-	virtual const float GetLinValue(const float OutputValue);
-	virtual const float GetOutputValueForText(const float LinValue);
-	virtual const float GetLinValueForText(const float OutputValue);
+	void SetSliderValue(float InSliderValue);
+	virtual const float GetOutputValue(const float InSliderValue);
+	virtual const float GetSliderValue(const float OutputValue);
+	virtual const float GetOutputValueForText(const float InSliderValue);
+	virtual const float GetSliderValueForText(const float OutputValue);
 	virtual void SetOutputRange(const FVector2D Range);
 
 	// Text label functions 
@@ -122,7 +134,7 @@ public:
 
 protected:
 	const FAudioRadialSliderStyle* Style;
-	TAttribute<float> Value;
+	TAttribute<float> SliderValue;
 	FRuntimeFloatCurve SliderCurve;
 	
 	TAttribute<FSlateColor> CenterBackgroundColor;
@@ -142,7 +154,7 @@ protected:
 
 	// Range for output 
 	FVector2D OutputRange = FVector2D(0.0f, 1.0f);
-	static const FVector2D LinearRange;
+	static const FVector2D NormalizedLinearSliderRange;
 
 	TSharedRef<SWidgetSwitcher> CreateLayoutWidgetSwitcher();
 };
@@ -156,23 +168,22 @@ class AUDIOWIDGETS_API SAudioVolumeRadialSlider
 public:
 	SAudioVolumeRadialSlider();
 	void Construct(const SAudioRadialSlider::FArguments& InArgs);
-	const float GetOutputValue(const float LinValue);
-	const float GetLinValue(const float OutputValue);
-	const float GetOutputValueForText(const float LinValue) override;
-	const float GetLinValueForText(const float OutputValue) override;
+	const float GetOutputValue(const float InSliderValue) override;
+	const float GetSliderValue(const float OutputValue) override;
+	const float GetOutputValueForText(const float InSliderValue) override;
+	const float GetSliderValueForText(const float OutputValue) override;
 	void SetUseLinearOutput(bool InUseLinearOutput);
 	void SetOutputRange(const FVector2D Range) override;
 
-private:
 	// Min/max possible values for output range, derived to avoid Audio::ConvertToLinear/dB functions returning NaN
-	static const float MinDbValue; 
-	static const float MaxDbValue; 
-
-	// Use linear (0 - 1) output value. Only applies to the output value reported by GetOutputValue(); the text displayed will still be in decibels. 
+	static const float MinDbValue;
+	static const float MaxDbValue;
+private:
+	// Use linear (converted from dB, not normalized) output value. Only applies to the output value reported by GetOutputValue(); the text displayed will still be in decibels. 
 	bool bUseLinearOutput = true;
 
-	const float GetDbValueFromLin(const float LinValue);
-	const float GetLinValueFromDb(const float DbValue);
+	const float GetDbValueFromSliderValue(const float InSliderValue);
+	const float GetSliderValueFromDb(const float DbValue);
 };
 
 /*
@@ -184,6 +195,6 @@ class AUDIOWIDGETS_API SAudioFrequencyRadialSlider
 public:
 	SAudioFrequencyRadialSlider();
 	void Construct(const SAudioRadialSlider::FArguments& InArgs);
-	const float GetOutputValue(const float LinValue);
-	const float GetLinValue(const float OutputValue);
+	const float GetOutputValue(const float InSliderValue);
+	const float GetSliderValue(const float OutputValue);
 };
\ No newline at end of file
diff --git a/Engine/Plugins/Runtime/AudioWidgets/Source/AudioWidgets/Public/SAudioSlider.h b/Engine/Plugins/Runtime/AudioWidgets/Source/AudioWidgets/Public/SAudioSlider.h
--- a/Engine/Plugins/Runtime/AudioWidgets/Source/AudioWidgets/Public/SAudioSlider.h
+++ b/Engine/Plugins/Runtime/AudioWidgets/Source/AudioWidgets/Public/SAudioSlider.h
@@ -28,7 +28,7 @@ class AUDIOWIDGETS_API SAudioSliderBase
 public:
 	SLATE_BEGIN_ARGS(SAudioSliderBase)
 	{
-		_Value = 0.0f;
+		_SliderValue = 0.0f;
 		_Orientation = Orient_Vertical;
 
 		const ISlateStyle* AudioWidgetsStyle = FSlateStyleRegistry::FindSlateStyle("AudioWidgetsStyle");
@@ -44,8 +44,8 @@ public:
 		/** The style used to draw the audio slider. */
 		SLATE_STYLE_ARGUMENT(FAudioSliderStyle, Style)
 
-		/** A value representing the audio slider value. */
-		SLATE_ATTRIBUTE(float, Value)
+		/** A value representing the normalized linear (0 - 1) audio slider value position. */
+		SLATE_ATTRIBUTE(float, SliderValue)
 
 		/** Whether the text label is always shown or only on hover. */
 		SLATE_ATTRIBUTE(bool, AlwaysShowLabel)
@@ -91,15 +91,15 @@ public:
 	 * @param InDeclaration A declaration from which to construct the widget.
 	 */
 	virtual void Construct(const SAudioSliderBase::FArguments& InDeclaration);
-	virtual const float GetOutputValue(const float LinValue);
-	virtual const float GetLinValue(const float OutputValue);
-	virtual const float GetOutputValueForText(const float LinValue);
-	virtual const float GetLinValueForText(const float OutputValue);
+	virtual const float GetOutputValue(const float InSliderValue);
+	virtual const float GetSliderValue(const float OutputValue);
+	virtual const float GetOutputValueForText(const float InSliderValue);
+	virtual const float GetSliderValueForText(const float OutputValue);
 	
 	/**
 	 * Set the slider's linear (0-1 normalized) value. 
 	 */
-	void SetValue(float LinValue);
+	void SetSliderValue(float InSliderValue);
 	FVector2D ComputeDesiredSize(float) const;
 	void SetDesiredSizeOverride(const FVector2D DesiredSize);
 
@@ -122,7 +122,7 @@ protected:
 	const FAudioSliderStyle* Style;
 
 	// Holds the slider's current linear value, from 0.0 - 1.0f
-	TAttribute<float> ValueAttribute;
+	TAttribute<float> SliderValueAttribute;
 	// Holds the slider's orientation
 	TAttribute<EOrientation> Orientation;
 	// Optional override for desired size 
@@ -143,7 +143,7 @@ protected:
 
 	// Range for output, currently only used for frequency sliders and sliders without curves
 	FVector2D OutputRange = FVector2D(0.0f, 1.0f);
-	static const FVector2D LinearRange;
+	static const FVector2D NormalizedLinearSliderRange;
 private:
 	FSlateBrush SliderBackgroundBrush;
 	FVector2D SliderBackgroundSize;
@@ -167,8 +167,8 @@ public:
 	void SetOutputToLinCurve(const TWeakObjectPtr<const UCurveFloat> OutputToLinCurve);
 	const TWeakObjectPtr<const UCurveFloat> GetOutputToLinCurve();
 	const TWeakObjectPtr<const UCurveFloat> GetLinToOutputCurve();
-	const float GetOutputValue(const float LinValue);
-	const float GetLinValue(const float OutputValue);
+	const float GetOutputValue(const float InSliderValue);
+	const float GetSliderValue(const float OutputValue);
 
 protected:
 	// Curves for mapping linear (0.0 - 1.0) to output (ex. dB for volume)  
@@ -186,10 +186,10 @@ public:
 	SAudioVolumeSlider();
 	void Construct(const SAudioSlider::FArguments& InDeclaration);
 
-	const float GetOutputValue(const float LinValue);
-	const float GetLinValue(const float OutputValue);
-	const float GetOutputValueForText(const float LinValue) override;
-	const float GetLinValueForText(const float OutputValue) override;
+	const float GetOutputValue(const float InSliderValue);
+	const float GetSliderValue(const float OutputValue);
+	const float GetOutputValueForText(const float InSliderValue) override;
+	const float GetSliderValueForText(const float OutputValue) override;
 	void SetUseLinearOutput(bool InUseLinearOutput);
 	void SetOutputRange(const FVector2D Range) override;
 
@@ -198,11 +198,11 @@ private:
 	static const float MinDbValue;
 	static const float MaxDbValue; 
 	
-	// Use linear (0 - 1) output value. Only applies to the output value reported by GetOutputValue(); the text displayed will still be in decibels. 
+	// Use linear (converted from dB, not normalized) output value. Only applies to the output value reported by GetOutputValue(); the text displayed will still be in decibels. 
 	bool bUseLinearOutput = true;
 
-	const float GetDbValueFromLin(const float LinValue);
-	const float GetLinValueFromDb(const float DbValue);
+	const float GetDbValueFromSliderValue(const float InSliderValue);
+	const float GetSliderValueFromDb(const float DbValue);
 };
 
 /*
@@ -214,6 +214,6 @@ class AUDIOWIDGETS_API SAudioFrequencySlider
 public:
 	SAudioFrequencySlider();
 	void Construct(const SAudioSlider::FArguments& InDeclaration);
-	const float GetOutputValue(const float LinValue);
-	const float GetLinValue(const float OutputValue);
+	const float GetOutputValue(const float InSliderValue);
+	const float GetSliderValue(const float OutputValue);
 };
diff --git a/Engine/Plugins/Runtime/Metasound/Source/MetasoundEditor/Private/MetasoundEditorGraphMemberDefaults.cpp b/Engine/Plugins/Runtime/Metasound/Source/MetasoundEditor/Private/MetasoundEditorGraphMemberDefaults.cpp
--- a/Engine/Plugins/Runtime/Metasound/Source/MetasoundEditor/Private/MetasoundEditorGraphMemberDefaults.cpp
+++ b/Engine/Plugins/Runtime/Metasound/Source/MetasoundEditor/Private/MetasoundEditorGraphMemberDefaults.cpp
@@ -183,9 +183,10 @@ void UMetasoundEditorGraphMemberDefaultFloat::PostEditChangeChainProperty(FPrope
 	else if (PropertyChangedEvent.GetPropertyName().IsEqual(GET_MEMBER_NAME_CHECKED(UMetasoundEditorGraphMemberDefaultFloat, WidgetType)) ||
 		PropertyChangedEvent.GetPropertyName().IsEqual(GET_MEMBER_NAME_CHECKED(UMetasoundEditorGraphMemberDefaultFloat, WidgetValueType)))
 	{
-		if (WidgetValueType == EMetasoundMemberDefaultWidgetValueType::Volume)
+		// Update VolumeWidgetDecibelRange based on current range (it might be stale)
+		if (WidgetValueType == EMetasoundMemberDefaultWidgetValueType::Volume && VolumeWidgetUseLinearOutput)
 		{
-			SetRange(FVector2D(-100.0f, 0.0f));
+			VolumeWidgetDecibelRange = FVector2D(Audio::ConvertToDecibels(Range.X), Audio::ConvertToDecibels(Range.Y));
 		}
 
 		// If the widget type is changed to none, we need to refresh clamping the value or not, since if the widget was a slider before, the value was clamped
@@ -195,9 +196,34 @@ void UMetasoundEditorGraphMemberDefaultFloat::PostEditChangeChainProperty(FPrope
 	{
 		if (PropertyChangedEvent.ChangeType != EPropertyChangeType::Interactive)
 		{
-			// if Range.X > Range.Y, set Range.Y to Range.X
-			Range.Y = FMath::Max(Range.X, Range.Y);
-			ForceRefresh();
+			if (WidgetType != EMetasoundMemberDefaultWidget::None &&
+				WidgetValueType == EMetasoundMemberDefaultWidgetValueType::Volume)
+			{
+				if (VolumeWidgetUseLinearOutput)
+				{
+					Range.X = FMath::Max(Range.X, Audio::ConvertToLinear(SAudioVolumeRadialSlider::MinDbValue));
+					Range.Y = FMath::Min(Range.Y, Audio::ConvertToLinear(SAudioVolumeRadialSlider::MaxDbValue));
+					if (Range.X > Range.Y)
+					{
+						Range.Y = Range.X;
+					}
+					VolumeWidgetDecibelRange = FVector2D(Audio::ConvertToDecibels(Range.X), Audio::ConvertToDecibels(Range.Y));
+					OnRangeChanged.Broadcast(VolumeWidgetDecibelRange);
+					SetDefault(FMath::Clamp(Default, Range.X, Range.Y));
+				}
+			}
+			else
+			{
+				// if Range.X > Range.Y, set Range.Y to Range.X
+				if (Range.X > Range.Y)
+				{
+					SetRange(FVector2D(Range.X, FMath::Max(Range.X, Range.Y)));
+				}
+				else
+				{
+					ForceRefresh();
+				}
+			}
 		}
 	}
 	else if (PropertyChangedEvent.GetPropertyName().IsEqual(GET_MEMBER_NAME_CHECKED(UMetasoundEditorGraphMemberDefaultFloat, ClampDefault)))
@@ -217,6 +243,32 @@ void UMetasoundEditorGraphMemberDefaultFloat::PostEditChangeChainProperty(FPrope
 		}
 		OnClampChanged.Broadcast(ClampDefault);
 	}
+	else if (PropertyChangedEvent.PropertyChain.GetActiveMemberNode()->GetValue()->GetFName().IsEqual(GET_MEMBER_NAME_CHECKED(UMetasoundEditorGraphMemberDefaultFloat, VolumeWidgetDecibelRange)))
+	{
+		if (PropertyChangedEvent.ChangeType != EPropertyChangeType::Interactive)
+		{
+			VolumeWidgetDecibelRange.X = FMath::Max(VolumeWidgetDecibelRange.X, SAudioVolumeRadialSlider::MinDbValue);
+			VolumeWidgetDecibelRange.Y = FMath::Min(VolumeWidgetDecibelRange.Y, SAudioVolumeRadialSlider::MaxDbValue);
+			SetRange(FVector2D(Audio::ConvertToLinear(VolumeWidgetDecibelRange.X), Audio::ConvertToLinear(VolumeWidgetDecibelRange.Y)));
+		}
+	}
+	else if (PropertyChangedEvent.GetPropertyName().IsEqual(GET_MEMBER_NAME_CHECKED(UMetasoundEditorGraphMemberDefaultFloat, VolumeWidgetUseLinearOutput)))
+	{
+		if (VolumeWidgetUseLinearOutput)
+		{
+			// Range and default are currently in dB, need to change to linear
+			float DbDefault = Default;
+			VolumeWidgetDecibelRange = Range;
+			Range = FVector2D(Audio::ConvertToLinear(VolumeWidgetDecibelRange.X), Audio::ConvertToLinear(VolumeWidgetDecibelRange.Y));
+			Default = Audio::ConvertToLinear(DbDefault);
+		}
+		else
+		{
+			// Range and default are currently linear, need to change to dB
+			Range = VolumeWidgetDecibelRange;
+			Default = Audio::ConvertToDecibels(Default);
+		}
+	}
 
 	UMetasoundEditorGraphMember* Member = GetParentMember();
 	if (ensure(Member))
diff --git a/Engine/Plugins/Runtime/Metasound/Source/MetasoundEditor/Private/MetasoundNodeDetailCustomization.cpp b/Engine/Plugins/Runtime/Metasound/Source/MetasoundEditor/Private/MetasoundNodeDetailCustomization.cpp
--- a/Engine/Plugins/Runtime/Metasound/Source/MetasoundEditor/Private/MetasoundNodeDetailCustomization.cpp
+++ b/Engine/Plugins/Runtime/Metasound/Source/MetasoundEditor/Private/MetasoundNodeDetailCustomization.cpp
@@ -132,21 +132,14 @@ namespace Metasound
 				DefaultValueHandle = Row->GetPropertyHandle();
 			}
 
-			// Apply the clamp range to the default value if 
-			// 1. not using a widget and ClampDefault is true or 
-			// 2. using a widget but not using a volume widget where VolumeWidgetUseLinearOutput is true 
-			const bool bVolumeWidgetWithLinearOutput = DefaultFloat->WidgetValueType == EMetasoundMemberDefaultWidgetValueType::Volume && DefaultFloat->VolumeWidgetUseLinearOutput;
+			// Apply the clamp range to the default value if not using a widget and ClampDefault is true
 			const bool bUsingWidget = DefaultFloat->WidgetType != EMetasoundMemberDefaultWidget::None;
-			const bool bShouldClampDefaultValue = (!bUsingWidget && DefaultFloat->ClampDefault) || (bUsingWidget && !bVolumeWidgetWithLinearOutput);
+			const bool bShouldClampDefaultValue = bUsingWidget || (!bUsingWidget && DefaultFloat->ClampDefault);
 
 			IDetailPropertyRow* ClampRow = DefaultCategoryBuilder->AddExternalObjectProperty(TArray<UObject*>({ DefaultFloat }), GET_MEMBER_NAME_CHECKED(UMetasoundEditorGraphMemberDefaultFloat, ClampDefault));
 			if (ensure(ClampRow))
 			{
 				DefaultRows.Add(ClampRow);
-				if (bVolumeWidgetWithLinearOutput)
-				{
-					DefaultFloat->SetRange(FVector2D(0.0f, 1.0f));
-				}
 
 				if (DefaultValueHandle.IsValid())
 				{
@@ -210,7 +203,10 @@ namespace Metasound
 				if (DefaultFloat->WidgetValueType == EMetasoundMemberDefaultWidgetValueType::Volume)
 				{
 					DefaultRows.Add(WidgetCategoryBuilder.AddExternalObjectProperty(TArray<UObject*>({ DefaultFloat }), GET_MEMBER_NAME_CHECKED(UMetasoundEditorGraphMemberDefaultFloat, VolumeWidgetUseLinearOutput)));
-
+					if (DefaultFloat->VolumeWidgetUseLinearOutput)
+					{
+						DefaultRows.Add(WidgetCategoryBuilder.AddExternalObjectProperty(TArray<UObject*>({ DefaultFloat }), GET_MEMBER_NAME_CHECKED(UMetasoundEditorGraphMemberDefaultFloat, VolumeWidgetDecibelRange)));
+					}
 				}
 			}
 
diff --git a/Engine/Plugins/Runtime/Metasound/Source/MetasoundEditor/Private/SMetasoundGraphNode.cpp b/Engine/Plugins/Runtime/Metasound/Source/MetasoundEditor/Private/SMetasoundGraphNode.cpp
--- a/Engine/Plugins/Runtime/Metasound/Source/MetasoundEditor/Private/SMetasoundGraphNode.cpp
+++ b/Engine/Plugins/Runtime/Metasound/Source/MetasoundEditor/Private/SMetasoundGraphNode.cpp
@@ -73,6 +73,12 @@ namespace Metasound
 					}
 				}
 			}
+
+			if (bIsInputWidgetTransacting)
+			{
+				GEditor->EndTransaction();
+				UE_LOG(LogMetaSound, Warning, TEXT("Unmatched MetaSound editor widget transaction."));
+			}
 		}
 
 		bool SMetaSoundGraphNode::IsVariableAccessor() const
@@ -652,10 +658,10 @@ namespace Metasound
 						{
 							if (InputWidget.IsValid())
 							{
-								if (!this->bIsInputWidgetTransacting)
+								if (!bIsInputWidgetTransacting)
 								{
 									GEditor->BeginTransaction(LOCTEXT("MetasoundGraphNode_MetasoundSetInputDefault", "Set MetaSound Input Default"));
-									this->bIsInputWidgetTransacting = true;
+									bIsInputWidgetTransacting = true;
 								}
 								GraphMember->GetOwningGraph()->GetMetasound()->Modify();
 								DefaultFloat->Modify();
@@ -778,22 +784,50 @@ namespace Metasound
 						}
 						else if (DefaultFloat->WidgetType == EMetasoundMemberDefaultWidget::RadialSlider)
 						{
+							auto OnRadialSliderMouseCaptureBeginLambda = [this]()
+							{
+								if (!bIsInputWidgetTransacting)
+								{
+									GEditor->BeginTransaction(LOCTEXT("MetasoundSetRadialSliderInputDefault", "Set MetaSound Input Default"));
+									bIsInputWidgetTransacting = true;
+								}
+							};
+
+							auto OnRadialSliderMouseCaptureEndLambda = [this]()
+							{
+								if (bIsInputWidgetTransacting)
+								{
+									GEditor->EndTransaction();
+									bIsInputWidgetTransacting = false;
+								}
+								else
+								{
+									UE_LOG(LogMetaSound, Warning, TEXT("Unmatched MetaSound editor widget transaction."));
+								}
+							};
+
 							// Create slider 
 							if (DefaultFloat->WidgetValueType == EMetasoundMemberDefaultWidgetValueType::Frequency)
 							{
 								SAssignNew(InputWidget, SAudioFrequencyRadialSlider)
-									.OnValueChanged_Lambda(OnValueChangedLambda);
+									.OnValueChanged_Lambda(OnValueChangedLambda)
+									.OnMouseCaptureBegin_Lambda(OnRadialSliderMouseCaptureBeginLambda)
+									.OnMouseCaptureEnd_Lambda(OnRadialSliderMouseCaptureEndLambda);
 							}
 							else if (DefaultFloat->WidgetValueType == EMetasoundMemberDefaultWidgetValueType::Volume)
 							{
 								SAssignNew(InputWidget, SAudioVolumeRadialSlider)
-									.OnValueChanged_Lambda(OnValueChangedLambda);
+									.OnValueChanged_Lambda(OnValueChangedLambda)
+									.OnMouseCaptureBegin_Lambda(OnRadialSliderMouseCaptureBeginLambda)
+									.OnMouseCaptureEnd_Lambda(OnRadialSliderMouseCaptureEndLambda);
 								StaticCastSharedPtr<SAudioVolumeRadialSlider>(InputWidget)->SetUseLinearOutput(DefaultFloat->VolumeWidgetUseLinearOutput);
 							}
 							else
 							{
 								SAssignNew(InputWidget, SAudioRadialSlider)
-									.OnValueChanged_Lambda(OnValueChangedLambda);
+									.OnValueChanged_Lambda(OnValueChangedLambda)
+									.OnMouseCaptureBegin_Lambda(OnRadialSliderMouseCaptureBeginLambda)
+									.OnMouseCaptureEnd_Lambda(OnRadialSliderMouseCaptureEndLambda);
 								InputWidget->SetShowUnitsText(false);
 							}
 							// Only vertical layout for radial slider
@@ -818,7 +852,7 @@ namespace Metasound
 
 						InputWidget->SetOutputRange(DefaultFloat->GetRange());
 						InputWidget->SetUnitsTextReadOnly(true);
-						InputWidget->SetValue(InputWidget->GetLinValue(DefaultFloat->GetDefault()));
+						InputWidget->SetSliderValue(InputWidget->GetSliderValue(DefaultFloat->GetDefault()));
 						InputWidget->SetVisibility(TAttribute<EVisibility>::Create([this]()
 						{
 							if (UMetasoundEditorGraphMemberNode* Node = GetMetaSoundMemberNode())
@@ -839,8 +873,8 @@ namespace Metasound
 						{
 							if (Widget.IsValid())
 							{
-								const float LinValue = Widget->GetLinValue(Value);
-								Widget->SetValue(LinValue);
+								const float SliderValue = Widget->GetSliderValue(Value);
+								Widget->SetSliderValue(SliderValue);
 							}
 						});
 
diff --git a/Engine/Plugins/Runtime/Metasound/Source/MetasoundEditor/Public/MetasoundEditorGraphMemberDefaults.h b/Engine/Plugins/Runtime/Metasound/Source/MetasoundEditor/Public/MetasoundEditorGraphMemberDefaults.h
--- a/Engine/Plugins/Runtime/Metasound/Source/MetasoundEditor/Public/MetasoundEditorGraphMemberDefaults.h
+++ b/Engine/Plugins/Runtime/Metasound/Source/MetasoundEditor/Public/MetasoundEditorGraphMemberDefaults.h
@@ -20,6 +20,7 @@
 #include "MetasoundPrimitives.h"
 #include "MetasoundVertex.h"
 #include "Misc/Guid.h"
+#include "SAudioRadialSlider.h"
 #include "Sound/SoundWave.h"
 #include "UObject/ObjectMacros.h"
 #include "UObject/SoftObjectPath.h"
@@ -156,10 +157,14 @@ public:
 	UPROPERTY(EditAnywhere, Category = Widget, meta=(DisplayName = "Value Type", EditCondition = "WidgetType != EMetasoundMemberDefaultWidget::None", EditConditionHides))
 	EMetasoundMemberDefaultWidgetValueType WidgetValueType = EMetasoundMemberDefaultWidgetValueType::Linear;
 
-	/** If true, output linear (0 - 1) value. Otherwise, output dB value. The volume widget itself will always display the value in dB. */
+	/** If true, output linear value. Otherwise, output dB value. The volume widget itself will always display the value in dB. The Default Value and Range are linear. */
 	UPROPERTY(EditAnywhere, Category = Widget, meta = (DisplayName = "Output Linear"))
 	bool VolumeWidgetUseLinearOutput = true;
 
+	/** Range in decibels. This will be converted to the linear range in the Default Value category. */
+	UPROPERTY(EditAnywhere, Category = Widget, meta = (DisplayName = "Range in dB", EditCondition = "VolumeWidgetUseLinearOutput", EditConditionHides))
+	FVector2D VolumeWidgetDecibelRange = FVector2D(SAudioVolumeRadialSlider::MinDbValue, 0.0f);
+
 	static FName GetDefaultPropertyName()
 	{
 		return GET_MEMBER_NAME_CHECKED(UMetasoundEditorGraphMemberDefaultFloat, Default);
