Audio Dashboard Implementation Checkpoint

diff --git a/Engine/Plugins/Runtime/AudioWidgets/AudioWidgets.uplugin b/Engine/Plugins/Runtime/AudioWidgets/AudioWidgets.uplugin
--- a/Engine/Plugins/Runtime/AudioWidgets/AudioWidgets.uplugin
+++ b/Engine/Plugins/Runtime/AudioWidgets/AudioWidgets.uplugin
@@ -11,7 +11,7 @@
 	"MarketplaceURL": "",
 	"SupportURL": "",
 	"CanContainContent": true,
-	"IsBetaVersion": true,
+	"IsBetaVersion": false,
 	"IsExperimentalVersion": false,
 	"Installed": false,
 	"Modules": [
@@ -20,5 +20,11 @@
 			"Type": "Runtime",
 			"LoadingPhase": "Default"
 		}
+	],
+	"Plugins": [
+		{
+			"Name": "AudioSynesthesia",
+			"Enabled":  true
+		}
 	]
 }
\ No newline at end of file
diff --git a/Engine/Plugins/Runtime/AudioWidgets/Source/AudioWidgets/AudioWidgets.Build.cs b/Engine/Plugins/Runtime/AudioWidgets/Source/AudioWidgets/AudioWidgets.Build.cs
--- a/Engine/Plugins/Runtime/AudioWidgets/Source/AudioWidgets/AudioWidgets.Build.cs
+++ b/Engine/Plugins/Runtime/AudioWidgets/Source/AudioWidgets/AudioWidgets.Build.cs
@@ -8,16 +8,20 @@ public class AudioWidgets : ModuleRules
 	{
 		PCHUsage = ModuleRules.PCHUsageMode.UseExplicitOrSharedPCHs;
 
-		PublicDependencyModuleNames.AddRange(
-		new string[]
+		PublicDependencyModuleNames.AddRange
+		(
+			new string[]
 			{
 				"Core",
-				"UMG"
+				"UMG",
+				"AudioAnalyzer",
+				"AudioSynesthesia"
 			}
 		);
 
-		PrivateDependencyModuleNames.AddRange(
-		new string[]
+		PrivateDependencyModuleNames.AddRange
+		(
+			new string[]
 			{
 				"CoreUObject",
 				"Engine",
@@ -25,7 +29,7 @@ public class AudioWidgets : ModuleRules
 				"SlateCore",
 				"InputCore",
 				"AdvancedWidgets",
-				"SignalProcessing",
+				"SignalProcessing"
 			}
 		);
 	}
diff --git a/Engine/Plugins/Runtime/AudioWidgets/Source/AudioWidgets/Private/AudioMeter.cpp b/Engine/Plugins/Runtime/AudioWidgets/Source/AudioWidgets/Private/AudioMeter.cpp
--- a/Engine/Plugins/Runtime/AudioWidgets/Source/AudioWidgets/Private/AudioMeter.cpp
+++ b/Engine/Plugins/Runtime/AudioWidgets/Source/AudioWidgets/Private/AudioMeter.cpp
@@ -2,21 +2,19 @@
 
 #include "AudioMeter.h"
 #include "SAudioMeter.h"
+#include "UObject/UObjectGlobals.h"
 #include "Widgets/DeclarativeSyntaxSupport.h"
 
 #include UE_INLINE_GENERATED_CPP_BY_NAME(AudioMeter)
 
 
 #define LOCTEXT_NAMESPACE "AUDIO_UMG"
-
-/////////////////////////////////////////////////////
-// UAudioMeter
-
 UAudioMeter::UAudioMeter(const FObjectInitializer& ObjectInitializer)
 	: Super(ObjectInitializer)
 {
 	Orientation = EOrientation::Orient_Vertical;
 
+	// TODO: Move to style set
 	BackgroundColor = FLinearColor(0.0075f, 0.0075f, 0.0075, 1.0f);
 	MeterBackgroundColor = FLinearColor(0.031f, 0.031f, 0.031f, 1.0f);
 	MeterValueColor = FLinearColor(0.025719f, 0.208333f, 0.069907f, 1.0f);
@@ -153,14 +151,107 @@ void UAudioMeter::SetMeterScaleLabelColor(FLinearColor InValue)
 }
 
 #if WITH_EDITOR
-
 const FText UAudioMeter::GetPaletteCategory()
 {
 	return LOCTEXT("Audio", "Audio");
 }
-
 #endif
 
-/////////////////////////////////////////////////////
+namespace AudioWidgets
+{
+	FAudioMeter::FAudioMeter()
+		: Widget(SNew(SAudioMeter)
+			.Orientation(EOrientation::Orient_Vertical)
+			.BackgroundColor(FLinearColor::Transparent)
+
+			// TODO: Move to editor style
+			.MeterBackgroundColor(FLinearColor(0.031f, 0.031f, 0.031f, 1.0f))
+			.MeterValueColor(FLinearColor(0.025719f, 0.208333f, 0.069907f, 1.0f))
+			.MeterPeakColor(FLinearColor(0.24349f, 0.708333f, 0.357002f, 1.0f))
+			.MeterClippingColor(FLinearColor(1.0f, 0.0f, 0.112334f, 1.0f))
+			.MeterScaleColor(FLinearColor(0.017642f, 0.017642f, 0.017642f, 1.0f))
+			.MeterScaleLabelColor(FLinearColor(0.442708f, 0.442708f, 0.442708f, 1.0f))
+		)
+	{
+	}
+
+	UAudioBus* FAudioMeter::GetAudioBus() const
+	{
+		return AudioBus.Get();
+	}
+
+	TSharedRef<SAudioMeter> FAudioMeter::GetWidget() const
+	{
+		return StaticCastSharedRef<SAudioMeter>(Widget->AsShared());
+	}
+
+	void FAudioMeter::Init(int32 InNumChannels, UWorld& InWorld)
+	{
+		check(InNumChannels > 0);
 
+		Teardown();
+
+		Settings = TStrongObjectPtr(NewObject<UMeterSettings>());
+		Settings->PeakHoldTime = 4000.0f;
+
+		Analyzer = TStrongObjectPtr(NewObject<UMeterAnalyzer>());
+		Analyzer->Settings = Settings.Get();
+
+		AudioBus = TStrongObjectPtr(NewObject<UAudioBus>());
+		AudioBus->AudioBusChannels = EAudioBusChannels(InNumChannels - 1);
+
+		ResultsDelegateHandle = Analyzer->OnLatestPerChannelMeterResultsNative.AddRaw(this, &FAudioMeter::OnMeterOutput);
+
+		WorldPtr = &InWorld;
+		Analyzer->StartAnalyzing(&InWorld, AudioBus.Get());
+
+		constexpr float DefaultMeterValue = -60.0f;
+		constexpr float DefaultPeakValue = -60.0f;
+		ChannelInfo.Init(FMeterChannelInfo{ DefaultMeterValue, DefaultPeakValue }, InNumChannels);
+
+		if (Widget.IsValid())
+		{
+			Widget->SetMeterChannelInfo(ChannelInfo);
+		}
+	}
+
+	void FAudioMeter::OnMeterOutput(UMeterAnalyzer* InMeterAnalyzer, int32 ChannelIndex, const FMeterResults& InMeterResults)
+	{
+		if (InMeterAnalyzer == Analyzer.Get())
+		{
+			FMeterChannelInfo NewChannelInfo;
+			NewChannelInfo.MeterValue = InMeterResults.MeterValue;
+			NewChannelInfo.PeakValue = InMeterResults.PeakValue;
+
+			if (ChannelIndex < ChannelInfo.Num())
+			{
+				ChannelInfo[ChannelIndex] = MoveTemp(NewChannelInfo);
+			}
+
+			// Update the widget if this is the last channel
+			if (ChannelIndex == ChannelInfo.Num() - 1)
+			{
+				if (Widget.IsValid())
+				{
+					Widget->SetMeterChannelInfo(ChannelInfo);
+				}
+			}
+		}
+	}
+
+	void FAudioMeter::Teardown()
+	{
+		if (Analyzer.IsValid() && ResultsDelegateHandle.IsValid())
+		{
+			Analyzer->StopAnalyzing();
+			Analyzer->OnLatestPerChannelMeterResultsNative.Remove(ResultsDelegateHandle);
+			Analyzer.Reset();
+			ResultsDelegateHandle.Reset();
+		}
+
+		AudioBus.Reset();
+		ChannelInfo.Reset();
+		Settings.Reset();
+	}
+} // namespace AudioWidgets
 #undef LOCTEXT_NAMESPACE
diff --git a/Engine/Plugins/Runtime/AudioWidgets/Source/AudioWidgets/Public/AudioMeter.h b/Engine/Plugins/Runtime/AudioWidgets/Source/AudioWidgets/Public/AudioMeter.h
--- a/Engine/Plugins/Runtime/AudioWidgets/Source/AudioWidgets/Public/AudioMeter.h
+++ b/Engine/Plugins/Runtime/AudioWidgets/Source/AudioWidgets/Public/AudioMeter.h
@@ -1,17 +1,28 @@
 // Copyright Epic Games, Inc. All Rights Reserved.
-
 #pragma once
 
-#include "CoreMinimal.h"
-#include "UObject/ObjectMacros.h"
-#include "Styling/SlateTypes.h"
-#include "Widgets/SWidget.h"
-#include "Components/Widget.h"
+#include "AudioAnalyzer.h"
 #include "AudioMeterStyle.h"
 #include "AudioMeterTypes.h"
+#include "Components/Widget.h"
+#include "Delegates/IDelegateInstance.h"
+#include "Engine/World.h"
+#include "Meter.h"
+#include "SAudioMeter.h"
+#include "Sound/AudioBus.h"
+#include "Styling/SlateTypes.h"
+#include "Templates/SharedPointer.h"
+#include "UObject/ObjectMacros.h"
+#include "UObject/WeakObjectPtrTemplates.h"
+#include "UObject/StrongObjectPtr.h"
+#include "Widgets/SWidget.h"
+
 #include "AudioMeter.generated.h"
 
+// Forward Declarations
 class SAudioMeter;
+class UWorld;
+
 
 /**
  * An audio meter widget.
@@ -137,4 +148,45 @@ protected:
 	// End of UWidget interface
 
 	PROPERTY_BINDING_IMPLEMENTATION(TArray<FMeterChannelInfo>, MeterChannelInfo);
-};
\ No newline at end of file
+};
+
+namespace AudioWidgets
+{
+	class AUDIOWIDGETS_API FAudioMeter
+	{
+	public:
+		FAudioMeter();
+
+		UAudioBus* GetAudioBus() const;
+
+		TSharedRef<SAudioMeter> GetWidget() const;
+
+		void Teardown();
+
+		void Init(int32 InNumChannels, UWorld& InWorld);
+
+	protected:
+		void OnMeterOutput(UMeterAnalyzer* InMeterAnalyzer, int32 ChannelIndex, const FMeterResults& InMeterResults);
+
+	private:
+		/** Metasound analyzer object. */
+		TStrongObjectPtr<UMeterAnalyzer> Analyzer;
+
+		/** The audio bus used for analysis. */
+		TStrongObjectPtr<UAudioBus> AudioBus;
+
+		/** Cached channel info for the meter. */
+		TArray<FMeterChannelInfo> ChannelInfo;
+
+		/** Handle for results delegate for MetaSound meter analyzer. */
+		FDelegateHandle ResultsDelegateHandle;
+
+		/** Meter settings. */
+		TStrongObjectPtr<UMeterSettings> Settings;
+
+		/** MetaSound Output Meter widget */
+		TSharedPtr<SAudioMeter> Widget;
+
+		TWeakObjectPtr<UWorld> WorldPtr;
+	};
+} // namespace AudioWidgets
diff --git a/Engine/Plugins/Runtime/Metasound/Source/MetasoundEditor/Private/MetasoundEditor.cpp b/Engine/Plugins/Runtime/Metasound/Source/MetasoundEditor/Private/MetasoundEditor.cpp
--- a/Engine/Plugins/Runtime/Metasound/Source/MetasoundEditor/Private/MetasoundEditor.cpp
+++ b/Engine/Plugins/Runtime/Metasound/Source/MetasoundEditor/Private/MetasoundEditor.cpp
@@ -11,7 +11,7 @@
 #include "EdGraph/EdGraphPin.h"
 #include "EdGraphUtilities.h"
 #include "Editor.h"
-#include "Styling/AppStyle.h"
+#include "Engine/World.h"
 #include "Framework/Application/SlateApplication.h"
 #include "Framework/Commands/GenericCommands.h"
 #include "Framework/Docking/TabManager.h"
@@ -55,6 +55,7 @@
 #include "SMetasoundPalette.h"
 #include "SNodePanel.h"
 #include "Stats/Stats.h"
+#include "Styling/AppStyle.h"
 #include "Templates/Function.h"
 #include "Templates/SharedPointer.h"
 #include "Widgets/Docking/SDockTab.h"
@@ -665,7 +666,7 @@ namespace Metasound
 
 		TSharedPtr<SWidget> FEditor::BuildAnalyzerWidget() const
 		{
-			if (!OutputMeter.IsValid() || !OutputMeter->GetWidget().IsValid())
+			if (!OutputMeter.IsValid())
 			{
 				return SNullWidget::NullWidget->AsShared();
 			}
@@ -691,7 +692,7 @@ namespace Metasound
 				.HAlign(HAlign_Center)
 				.VAlign(VAlign_Fill)
 				[
-					OutputMeter->GetWidget().ToSharedRef()
+					OutputMeter->GetWidget()
 				]
 			];
 		}
@@ -1129,9 +1130,17 @@ namespace Metasound
 			{
 				if (!OutputMeter.IsValid())
 				{
-					OutputMeter = MakeShared<FEditorMeter>();
+					OutputMeter = MakeShared<AudioWidgets::FAudioMeter>();
+				}
+
+				if (ensure(GEditor))
+				{
+					UWorld* EditorWorld = GEditor->GetEditorWorldContext().World();
+					if (ensure(EditorWorld))
+					{
+						OutputMeter->Init(MetaSoundSource->NumChannels, *EditorWorld);
+					}
 				}
-				OutputMeter->Init(MetaSoundSource->NumChannels);
 			}
 			else
 			{
@@ -1583,10 +1592,8 @@ namespace Metasound
 								}
 								else
 								{
-									if (OutputMeter->GetWidget().IsValid())
-									{
-										OutputMeter->GetWidget()->bIsActiveTimerRegistered = false;
-									}
+									TSharedRef<SAudioMeter> MeterRef = OutputMeter->GetWidget();
+									MeterRef->bIsActiveTimerRegistered = false;
 									return EActiveTimerReturnType::Stop;
 								}
 							})
diff --git a/Engine/Plugins/Runtime/Metasound/Source/MetasoundEditor/Private/MetasoundEditor.h b/Engine/Plugins/Runtime/Metasound/Source/MetasoundEditor/Private/MetasoundEditor.h
--- a/Engine/Plugins/Runtime/Metasound/Source/MetasoundEditor/Private/MetasoundEditor.h
+++ b/Engine/Plugins/Runtime/Metasound/Source/MetasoundEditor/Private/MetasoundEditor.h
@@ -1,18 +1,18 @@
 // Copyright Epic Games, Inc. All Rights Reserved.
 #pragma once
 
-#include "IMetasoundEditor.h"
+#include "AudioMeter.h"
 #include "AudioMeterStyle.h"
 #include "EdGraph/EdGraphNode.h"
 #include "EditorUndoClient.h"
 #include "Framework/Commands/UICommandList.h"
 #include "GraphEditor.h"
 #include "IDetailsView.h"
+#include "IMetasoundEditor.h"
 #include "Logging/TokenizedMessage.h"
 #include "Math/UnrealMathUtility.h"
 #include "MetasoundEditorGraphConnectionManager.h"
 #include "MetasoundEditorGraphNode.h"
-#include "MetasoundEditorMeter.h"
 #include "MetasoundFrontend.h"
 #include "MetasoundFrontendController.h"
 #include "Meter.h"
@@ -418,7 +418,7 @@ namespace Metasound
 			TSharedPtr<SGraphActionMenu> GraphMembersMenu;
 
 			/** Meter used in the analyzer tab for auditioning preview output. */
-			TSharedPtr<FEditorMeter> OutputMeter;
+			TSharedPtr<AudioWidgets::FAudioMeter> OutputMeter;
 
 			/** Palette of Node types */
 			TSharedPtr<SMetasoundPalette> Palette;
diff --git a/Engine/Plugins/Runtime/Metasound/Source/MetasoundEditor/Private/MetasoundEditorMeter.cpp b/Engine/Plugins/Runtime/Metasound/Source/MetasoundEditor/Private/MetasoundEditorMeter.cpp
deleted file mode 100644
diff --git a/Engine/Plugins/Runtime/Metasound/Source/MetasoundEditor/Public/MetasoundEditorMeter.h b/Engine/Plugins/Runtime/Metasound/Source/MetasoundEditor/Public/MetasoundEditorMeter.h
deleted file mode 100644
diff --git a/Engine/Source/Runtime/AudioAnalyzer/Classes/AudioAnalyzer.h b/Engine/Source/Runtime/AudioAnalyzer/Classes/AudioAnalyzer.h
--- a/Engine/Source/Runtime/AudioAnalyzer/Classes/AudioAnalyzer.h
+++ b/Engine/Source/Runtime/AudioAnalyzer/Classes/AudioAnalyzer.h
@@ -94,7 +94,7 @@ public:
 
 	/** Stops analyzing audio. */
 	UFUNCTION(BlueprintCallable, BlueprintCosmetic, Category = AudioAnalyzer, meta = (WorldContext = "WorldContextObject"))
-	void StopAnalyzing(const UObject* WorldContextObject);
+	void StopAnalyzing(const UObject* WorldContextObject = nullptr);
 
 	/**
 	 * Implementations can override this method to create settings objects
