[WaveformEditor] Make drawing utils more generic and port them to AudioWidgets + AudioWidgetsTests

diff --git a/Engine/Plugins/Editor/WaveformEditor/Source/WaveformEditorWidgets/Private/SWaveformViewer.cpp b/Engine/Plugins/Editor/WaveformEditor/Source/WaveformEditorWidgets/Private/SWaveformViewer.cpp
--- a/Engine/Plugins/Editor/WaveformEditor/Source/WaveformEditorWidgets/Private/SWaveformViewer.cpp
+++ b/Engine/Plugins/Editor/WaveformEditor/Source/WaveformEditorWidgets/Private/SWaveformViewer.cpp
@@ -2,6 +2,7 @@
 
 #include "SWaveformViewer.h"
 
+#include "TimeSeriesDrawingUtils.h"
 #include "WaveformEditorRenderData.h"
 #include "WaveformEditorSlateTypes.h"
 
@@ -155,7 +156,7 @@ void SWaveformViewer::Tick(const FGeometry& AllottedGeometry, const double InCur
 
 		if (WaveformDrawMode == EWaveformDrawMode::BinnedPeaks)
 		{
-			WaveformDrawingUtils::GetBinnedPeaksFromWaveformRawData(CachedPeaks, CachedPixelWidth, SampleData.GetData(), RenderData->GetNumSamples(), RenderData->GetSampleRate(), RenderData->GetNumChannels(), DisplayRange.GetLowerBoundValue(), DisplayRange.GetUpperBoundValue());
+			TimeSeriesDrawingUtils::GroupInterleavedSampledTSIntoMinMaxBins(CachedPeaks, CachedPixelWidth, SampleData.GetData(), RenderData->GetNumSamples(), RenderData->GetSampleRate(), RenderData->GetNumChannels(), DisplayRange.GetLowerBoundValue(), DisplayRange.GetUpperBoundValue());
 			GenerateWaveformBinsCoordinates(CachedBinsDrawCoordinates, CachedPeaks, AllottedGeometry);
 		}
 		else
@@ -181,7 +182,7 @@ void SWaveformViewer::Tick(const FGeometry& AllottedGeometry, const double InCur
 	}
 }
 
-void SWaveformViewer::GenerateWaveformBinsCoordinates(TArray<FSamplesBinCoordinates>& OutDrawCoordinates, const TArray<WaveformDrawingUtils::SampleRange>& InWaveformPeaks, const FGeometry& InAllottedGeometry, const float VerticalZoomFactor /*= 1.f*/)
+void SWaveformViewer::GenerateWaveformBinsCoordinates(TArray<FSamplesBinCoordinates>& OutDrawCoordinates, const TArray<WaveformViewerWidget::SampleRange>& InWaveformPeaks, const FGeometry& InAllottedGeometry, const float VerticalZoomFactor /*= 1.f*/)
 {
 	const int NChannels = RenderData->GetNumChannels();
 
@@ -205,7 +206,7 @@ void SWaveformViewer::GenerateWaveformBinsCoordinates(TArray<FSamplesBinCoordina
 		for (uint32 Pixel = 0; Pixel < PixelWidth; ++Pixel)
 		{
 			uint32 PeakIndex = Pixel * NChannels + Channel;
-			const WaveformDrawingUtils::SampleRange& SamplePeaks = InWaveformPeaks[PeakIndex];
+			const WaveformViewerWidget::SampleRange& SamplePeaks = InWaveformPeaks[PeakIndex];
 
 			const float SampleMaxScaled = SamplePeaks.GetUpperBoundValue() * HeightScale > MinScaledSamplePeakValue ? SamplePeaks.GetUpperBoundValue() * HeightScale : MinWaveformBinHeight;
 			const float SampleMinScaled = SamplePeaks.GetLowerBoundValue() * HeightScale < -1 * MinScaledSamplePeakValue ? SamplePeaks.GetLowerBoundValue() * HeightScale : -1.f * MinWaveformBinHeight;
diff --git a/Engine/Plugins/Editor/WaveformEditor/Source/WaveformEditorWidgets/Private/SWaveformViewer.h b/Engine/Plugins/Editor/WaveformEditor/Source/WaveformEditorWidgets/Private/SWaveformViewer.h
--- a/Engine/Plugins/Editor/WaveformEditor/Source/WaveformEditorWidgets/Private/SWaveformViewer.h
+++ b/Engine/Plugins/Editor/WaveformEditor/Source/WaveformEditorWidgets/Private/SWaveformViewer.h
@@ -4,12 +4,16 @@
 
 #include "Styling/SlateStyleRegistry.h"
 #include "Styling/SlateWidgetStyleAsset.h"
-#include "WaveformDrawingUtils.h"
 #include "WaveformEditorGridData.h"
 #include "WaveformEditorSlateTypes.h"
 #include "WaveformEditorTransportCoordinator.h"
 #include "Widgets/SLeafWidget.h"
 
+namespace WaveformViewerWidget
+{
+	typedef TRange<int16> SampleRange;
+}
+
 class FWaveformEditorRenderData;
 struct FWaveformEditorWidgetStyleBase;
 struct FWaveformViewerStyle;
@@ -62,7 +66,7 @@ private:
 	virtual FVector2D ComputeDesiredSize(float) const override;
 	virtual void Tick(const FGeometry& AllottedGeometry, const double InCurrentTime, const float InDeltaTime) override;
 
-	void GenerateWaveformBinsCoordinates(TArray<FSamplesBinCoordinates>& OutDrawCoordinates, const TArray<WaveformDrawingUtils::SampleRange>& InWaveformPeaks, const FGeometry& InAllottedGeometry, const float VerticalZoomFactor = 1.f);
+	void GenerateWaveformBinsCoordinates(TArray<FSamplesBinCoordinates>& OutDrawCoordinates, const TArray<WaveformViewerWidget::SampleRange>& InWaveformPeaks, const FGeometry& InAllottedGeometry, const float VerticalZoomFactor = 1.f);
 	void GenerateWaveformSamplesCoordinates(TArray<FVector2D>& OutDrawCoordinates, const FGeometry& InAllottedGeometry, const float VerticalZoomFactor = 1.f);
 	void DrawGridLines(const FGeometry& AllottedGeometry, FSlateWindowElementList& OutDrawElements, int32& LayerId) const;
 
@@ -92,7 +96,7 @@ private:
 
 	uint32 CachedPixelWidth = 0; 
 	float CachedPixelHeight = 0.f;
-	TArray<WaveformDrawingUtils::SampleRange> CachedPeaks;
+	TArray<WaveformViewerWidget::SampleRange> CachedPeaks;
 	TArray<FSamplesBinCoordinates> CachedBinsDrawCoordinates;
 	TArray<FVector2D> CachedSampleDrawCoordinates;
 
diff --git a/Engine/Plugins/Editor/WaveformEditor/Source/WaveformEditorWidgets/Private/WaveformDrawingUtils.cpp b/Engine/Plugins/Editor/WaveformEditor/Source/WaveformEditorWidgets/Private/WaveformDrawingUtils.cpp
deleted file mode 100644
diff --git a/Engine/Plugins/Editor/WaveformEditor/Source/WaveformEditorWidgets/Public/WaveformDrawingUtils.h b/Engine/Plugins/Editor/WaveformEditor/Source/WaveformEditorWidgets/Public/WaveformDrawingUtils.h
deleted file mode 100644
diff --git a/Engine/Plugins/Editor/WaveformEditor/Source/WaveformEditorWidgets/WaveformEditorWidgets.build.cs b/Engine/Plugins/Editor/WaveformEditor/Source/WaveformEditorWidgets/WaveformEditorWidgets.build.cs
--- a/Engine/Plugins/Editor/WaveformEditor/Source/WaveformEditorWidgets/WaveformEditorWidgets.build.cs
+++ b/Engine/Plugins/Editor/WaveformEditor/Source/WaveformEditorWidgets/WaveformEditorWidgets.build.cs
@@ -14,6 +14,7 @@ namespace UnrealBuildTool.Rules
 			PublicDependencyModuleNames.AddRange(
 				new string[] {
 					"AudioExtensions",
+					"AudioWidgets",
 					"ContentBrowser",
 					"Core",
 					"CoreUObject",
diff --git a/Engine/Plugins/Editor/WaveformEditor/WaveformEditor.uplugin b/Engine/Plugins/Editor/WaveformEditor/WaveformEditor.uplugin
--- a/Engine/Plugins/Editor/WaveformEditor/WaveformEditor.uplugin
+++ b/Engine/Plugins/Editor/WaveformEditor/WaveformEditor.uplugin
@@ -40,6 +40,10 @@
 		{
 			"Name": "AudioSynesthesia",
 			"Enabled": true
+		},
+		{
+			"Name": "AudioWidgets",
+			"Enabled" :  true
 		}
 	]
 }
\ No newline at end of file
diff --git a/Engine/Plugins/Runtime/AudioWidgets/Source/AudioWidgets/Public/TimeSeriesDrawingUtils.h b/Engine/Plugins/Runtime/AudioWidgets/Source/AudioWidgets/Public/TimeSeriesDrawingUtils.h
new file mode 100644
--- /dev/null
+++ b/Engine/Plugins/Runtime/AudioWidgets/Source/AudioWidgets/Public/TimeSeriesDrawingUtils.h
@@ -0,0 +1,90 @@
+// Copyright Epic Games, Inc. All Rights Reserved.
+
+#pragma once
+
+#include "Math/Range.h"
+
+namespace TimeSeriesDrawingUtils
+{
+	/**
+	* Groups samples evenly into a number of desired bins.
+	* Each bin contains the min and max values of the grouped samples.
+	*
+	* @param OutBins			TArray where the bins will be written to
+	* @param NumDesiredBins		Number of output bins
+	* @param RawDataPtr			Ptr to the beginning of the samples
+	* @param TotalNumSamples	The total number of samples of the input time series
+	* @param NChannels			Number of interleaved channels in the time series
+	* @param StartRatio			The ratio of the total number of frames (in a range of 0-1) at which grouping should start
+	* @param EndRatio			The ratio of the total number of frames (in a range of 0-1) at which grouping should end.
+	*
+	*/
+	template<typename SamplesType>
+	void GroupInterleavedSamplesIntoMinMaxBins(TArray<TRange<SamplesType>>& OutBins, const uint32 NumDesiredBins, const SamplesType* RawDataPtr, const uint32 TotalNumSamples, const uint16 NChannels = 1, const double StartRatio = 0.0, double EndRatio = 1.0)
+	{
+		check(StartRatio >= 0.f && StartRatio < EndRatio)
+		check(EndRatio <= 1.f)
+
+		uint32 NumPeaks = NumDesiredBins * NChannels;
+		OutBins.SetNumUninitialized(NumPeaks);
+		const uint32 NumFrames = TotalNumSamples / NChannels;
+
+		double FramesPerBin = ((EndRatio - StartRatio) * NumFrames) / NumDesiredBins;
+		double IterationStartFrame = StartRatio * NumFrames;
+		uint32 FirstBinnedFrame = FMath::RoundToInt(StartRatio * NumFrames);
+
+		for (uint32 Peak = 0; Peak < NumPeaks; Peak += NChannels)
+		{
+			uint32 LastBinnedFrame = IterationStartFrame + FramesPerBin;
+
+			for (uint16 Channel = 0; Channel < NChannels; ++Channel)
+			{
+				SamplesType MaxSampleValue = TNumericLimits<SamplesType>::Min();
+				SamplesType MinSampleValue = TNumericLimits<SamplesType>::Max();
+
+				for (uint32 Frame = FirstBinnedFrame; Frame < LastBinnedFrame; ++Frame)
+				{
+					const uint32 SampleIndex = Frame * NChannels + Channel;
+					check(SampleIndex < TotalNumSamples);
+					int16 SampleValue = RawDataPtr[SampleIndex];
+
+					MinSampleValue = FMath::Min(MinSampleValue, SampleValue);
+					MaxSampleValue = FMath::Max(MaxSampleValue, SampleValue);
+
+				}
+
+				OutBins[Peak + Channel] = TRange<SamplesType>(MinSampleValue, MaxSampleValue);
+			}
+
+			IterationStartFrame += FramesPerBin;
+			FirstBinnedFrame = LastBinnedFrame;
+		}
+	}
+
+	/**
+	* Groups samples of a time series into an equal number of desired bins.
+	* Each bin contains the min and max values of the grouped samples.
+	* 
+	* @param OutBins			TArray where the bins will be written to 
+	* @param NumDesiredBins		Number of output bins
+	* @param RawDataPtr			Ptr to the beginning of the samples 
+	* @param TotalNumSamples	The total number of samples of the input time series 
+	* @param SampleRate			SampleRate of the time series 
+	* @param NChannels			Number of interleaved channels in the time series
+	* @param StartTime			The time at which grouping should start
+	* @param EndTime 			The time at which grouping should end. 
+	*							
+	* Note: With a negative EndTime, the method will calculate automatically the EndTime by doing 
+	* TotalNumSamples / (SampleRate * NChannels)
+	*							
+	*/
+	template<typename SamplesType>
+	void GroupInterleavedSampledTSIntoMinMaxBins(TArray<TRange<SamplesType>>& OutBins, const uint32 NumDesiredBins, const SamplesType* RawDataPtr, const uint32 TotalNumSamples, const uint32 SampleRate, const uint16 NChannels = 1, const float StartTime = 0.f, float EndTime = -1.f)
+	{
+		const double TotalTime = TotalNumSamples / ((float) SampleRate * NChannels);
+		const double StartRatio = StartTime / TotalTime;
+		const double EndRatio = EndTime >= 0.f ? FMath::Clamp(EndTime / TotalTime, StartRatio, 1.0) : 1.0;
+
+		GroupInterleavedSamplesIntoMinMaxBins(OutBins, NumDesiredBins, RawDataPtr, TotalNumSamples, NChannels, StartRatio, EndRatio);
+	}
+}
\ No newline at end of file
