[Engine/Plugins]

diff --git a/Engine/Plugins/Online/OnlineSubsystem/Source/Public/OnlineIdentityErrors.h b/Engine/Plugins/Online/OnlineSubsystem/Source/Public/OnlineIdentityErrors.h
--- a/Engine/Plugins/Online/OnlineSubsystem/Source/Public/OnlineIdentityErrors.h
+++ b/Engine/Plugins/Online/OnlineSubsystem/Source/Public/OnlineIdentityErrors.h
@@ -3,7 +3,7 @@
 #pragma once
 
 
-#include "OnlineError.h"
+#include "OnlineError.h" // IWYU pragma: keep
 #define LOCTEXT_NAMESPACE "OnlineIdentity"
 #define ONLINE_ERROR_NAMESPACE "errors.com.epicgames.oss.identity"
 
diff --git a/Engine/Plugins/Online/OnlineSubsystem/Source/Public/OnlineJsonSerializer.h b/Engine/Plugins/Online/OnlineSubsystem/Source/Public/OnlineJsonSerializer.h
--- a/Engine/Plugins/Online/OnlineSubsystem/Source/Public/OnlineJsonSerializer.h
+++ b/Engine/Plugins/Online/OnlineSubsystem/Source/Public/OnlineJsonSerializer.h
@@ -2,9 +2,9 @@
 
 #pragma once
 
-#include "CoreMinimal.h"
 #include "Serialization/JsonSerializerMacros.h"
-#include "OnlineKeyValuePair.h"
+
+class FVariantData;
 
 #define BEGIN_ONLINE_JSON_SERIALIZER																BEGIN_JSON_SERIALIZER
 #define END_ONLINE_JSON_SERIALIZER																	END_JSON_SERIALIZER
@@ -28,3 +28,8 @@ typedef FJsonSerializerWriter<> FOnlineJsonSerializerWriter;
 
 typedef TMap<FString, FVariantData> FJsonSerializeableKeyValueMapVariant;
 
+
+#if UE_ENABLE_INCLUDE_ORDER_DEPRECATED_IN_5_2
+#include "CoreMinimal.h"
+#include "OnlineKeyValuePair.h"
+#endif
diff --git a/Engine/Plugins/Online/OnlineSubsystem/Source/Public/OnlineKeyValuePair.h b/Engine/Plugins/Online/OnlineSubsystem/Source/Public/OnlineKeyValuePair.h
--- a/Engine/Plugins/Online/OnlineSubsystem/Source/Public/OnlineKeyValuePair.h
+++ b/Engine/Plugins/Online/OnlineSubsystem/Source/Public/OnlineKeyValuePair.h
@@ -1,8 +1,10 @@
 // Copyright Epic Games, Inc. All Rights Reserved.
 #pragma once
 
-#include "CoreMinimal.h"
-#include "OnlineSubsystemPackage.h"
+#include "Templates/SharedPointer.h"
+
+class FJsonObject;
+class UStruct;
 
 namespace EOnlineKeyValuePairDataType
 {
@@ -628,3 +630,8 @@ private:
 	static bool ConvertScalarVariantToFProperty(const FVariantData* Variant, FProperty* Property, void* OutValue, int64 CheckFlags, int64 SkipFlags);
 };
 
+
+#if UE_ENABLE_INCLUDE_ORDER_DEPRECATED_IN_5_2
+#include "CoreMinimal.h"
+#include "OnlineSubsystemPackage.h"
+#endif
diff --git a/Engine/Plugins/Online/OnlineSubsystem/Source/Public/OnlineNotification.h b/Engine/Plugins/Online/OnlineSubsystem/Source/Public/OnlineNotification.h
--- a/Engine/Plugins/Online/OnlineSubsystem/Source/Public/OnlineNotification.h
+++ b/Engine/Plugins/Online/OnlineSubsystem/Source/Public/OnlineNotification.h
@@ -2,9 +2,10 @@
 
 #pragma once
 
-#include "CoreMinimal.h"
-#include "Online/CoreOnline.h"
-#include "JsonObjectWrapper.h"
+#include "Online/CoreOnlineFwd.h"
+
+class FJsonObject;
+class UStruct;
 
 class FJsonValue;
 
@@ -97,3 +98,9 @@ struct ONLINESUBSYSTEM_API FOnlineNotification
 	/** String representing the client_request_id for this notification. Used to tie a server request back to the client. Can be empty. */
 	FString ClientRequestIdStr;
 };
+
+#if UE_ENABLE_INCLUDE_ORDER_DEPRECATED_IN_5_2
+#include "CoreMinimal.h"
+#include "JsonObjectWrapper.h"
+#include "Online/CoreOnline.h"
+#endif
diff --git a/Engine/Plugins/Online/OnlineSubsystem/Source/Public/OnlineNotificationHandler.h b/Engine/Plugins/Online/OnlineSubsystem/Source/Public/OnlineNotificationHandler.h
--- a/Engine/Plugins/Online/OnlineSubsystem/Source/Public/OnlineNotificationHandler.h
+++ b/Engine/Plugins/Online/OnlineSubsystem/Source/Public/OnlineNotificationHandler.h
@@ -2,8 +2,7 @@
 
 #pragma once
 
-#include "CoreMinimal.h"
-#include "OnlineSubsystemPackage.h"
+#include "Delegates/Delegate.h"
 
 class FUniqueNetId;
 struct FOnlineNotification;
@@ -93,3 +92,8 @@ public:
 typedef TSharedPtr<FOnlineNotificationHandler, ESPMode::ThreadSafe> FOnlineNotificationHandlerPtr;
 typedef TSharedRef<FOnlineNotificationHandler, ESPMode::ThreadSafe> FOnlineNotificationHandlerRef;
 typedef TWeakPtr<FOnlineNotificationHandler, ESPMode::ThreadSafe> FOnlineNotificationHandlerWeakPtr;
+
+#if UE_ENABLE_INCLUDE_ORDER_DEPRECATED_IN_5_2
+#include "CoreMinimal.h"
+#include "OnlineSubsystemPackage.h"
+#endif
diff --git a/Engine/Plugins/Online/OnlineSubsystem/Source/Public/OnlineNotificationTransportManager.h b/Engine/Plugins/Online/OnlineSubsystem/Source/Public/OnlineNotificationTransportManager.h
--- a/Engine/Plugins/Online/OnlineSubsystem/Source/Public/OnlineNotificationTransportManager.h
+++ b/Engine/Plugins/Online/OnlineSubsystem/Source/Public/OnlineNotificationTransportManager.h
@@ -2,9 +2,7 @@
 
 #pragma once
 
-#include "CoreMinimal.h"
 #include "Interfaces/OnlineNotificationTransportInterface.h"
-#include "OnlineSubsystemPackage.h"
 
 struct FOnlineNotification;
 
@@ -65,3 +63,8 @@ public:
 };
 
 typedef TSharedPtr<FOnlineNotificationTransportManager, ESPMode::ThreadSafe> FOnlineNotificationTransportManagerPtr;
+
+#if UE_ENABLE_INCLUDE_ORDER_DEPRECATED_IN_5_2
+#include "CoreMinimal.h"
+#include "OnlineSubsystemPackage.h"
+#endif
diff --git a/Engine/Plugins/Online/OnlineSubsystem/Source/Public/OnlineSessionSettings.h b/Engine/Plugins/Online/OnlineSubsystem/Source/Public/OnlineSessionSettings.h
--- a/Engine/Plugins/Online/OnlineSubsystem/Source/Public/OnlineSessionSettings.h
+++ b/Engine/Plugins/Online/OnlineSubsystem/Source/Public/OnlineSessionSettings.h
@@ -2,14 +2,10 @@
 
 #pragma once
 
-#include "CoreMinimal.h"
-#include "Online/CoreOnline.h"
 #include "OnlineSubsystemTypes.h"
 #include "OnlineKeyValuePair.h"
-#include "OnlineSubsystemPackage.h"
 
 // NOTE:  Session setting names have been moved to OnlineBase in OnlineSessionNames.h
-#include "Online/OnlineSessionNames.h"
 
 #define INVALID_SESSION_SETTING_ID -1
 
@@ -662,3 +658,9 @@ void ONLINESUBSYSTEM_API DumpSession(const FOnlineSession* Session);
  */
 void ONLINESUBSYSTEM_API DumpSessionSettings(const FOnlineSessionSettings* SessionSettings);
 
+
+#if UE_ENABLE_INCLUDE_ORDER_DEPRECATED_IN_5_2
+#include "CoreMinimal.h"
+#include "Online/OnlineSessionNames.h"
+#include "OnlineSubsystemPackage.h"
+#endif
diff --git a/Engine/Plugins/Online/OnlineSubsystem/Source/Public/OnlineSubsystem.h b/Engine/Plugins/Online/OnlineSubsystem/Source/Public/OnlineSubsystem.h
--- a/Engine/Plugins/Online/OnlineSubsystem/Source/Public/OnlineSubsystem.h
+++ b/Engine/Plugins/Online/OnlineSubsystem/Source/Public/OnlineSubsystem.h
@@ -2,15 +2,13 @@
 
 #pragma once
 
-#include "CoreMinimal.h"
-#include "Online/CoreOnlineFwd.h"
-#include "Stats/Stats.h"
 #include "Modules/ModuleManager.h"
 #include "OnlineSubsystemModule.h"
 #include "OnlineSubsystemTypes.h"
 #include "OnlineDelegateMacros.h"
-#include "OnlineSubsystemNames.h"
-#include "HAL/LowLevelMemTracker.h"
+#include "OnlineSubsystemNames.h" // IWYU pragma: keep
+
+class IModuleInterface;
 
 class FOnlineNotificationHandler;
 class FOnlineNotificationTransportManager;
@@ -692,3 +690,7 @@ ONLINESUBSYSTEM_API int32 GetBeaconPortFromSessionSettings(const class FOnlineSe
 
 /** Temp solution for some hardcoded access to logged in user 0, please avoid using this */
 ONLINESUBSYSTEM_API FUniqueNetIdPtr GetFirstSignedInUser(IOnlineIdentityPtr IdentityInt);
+
+#if UE_ENABLE_INCLUDE_ORDER_DEPRECATED_IN_5_2
+#include "CoreMinimal.h"
+#endif
diff --git a/Engine/Plugins/Online/OnlineSubsystem/Source/Public/OnlineSubsystemImpl.h b/Engine/Plugins/Online/OnlineSubsystem/Source/Public/OnlineSubsystemImpl.h
--- a/Engine/Plugins/Online/OnlineSubsystem/Source/Public/OnlineSubsystemImpl.h
+++ b/Engine/Plugins/Online/OnlineSubsystem/Source/Public/OnlineSubsystemImpl.h
@@ -2,11 +2,9 @@
 
 #pragma once
 
-#include "CoreMinimal.h"
 #include "HAL/IConsoleManager.h"
 #include "OnlineSubsystem.h"
 #include "Containers/Queue.h"
-#include "OnlineSubsystemPackage.h"
 #include "Containers/Ticker.h"
 
 struct FOnlineError;
@@ -164,3 +162,8 @@ public:
 	static const FName DefaultInstanceName;
 };
 
+
+#if UE_ENABLE_INCLUDE_ORDER_DEPRECATED_IN_5_2
+#include "CoreMinimal.h"
+#include "OnlineSubsystemPackage.h"
+#endif
diff --git a/Engine/Plugins/Online/OnlineSubsystem/Source/Public/OnlineSubsystemModule.h b/Engine/Plugins/Online/OnlineSubsystem/Source/Public/OnlineSubsystemModule.h
--- a/Engine/Plugins/Online/OnlineSubsystem/Source/Public/OnlineSubsystemModule.h
+++ b/Engine/Plugins/Online/OnlineSubsystem/Source/Public/OnlineSubsystemModule.h
@@ -2,8 +2,10 @@
 
 #pragma once
 
-#include "CoreMinimal.h"
 #include "Modules/ModuleInterface.h"
+#include "Templates/SharedPointer.h"
+
+class IOnlineFactory;
 
 class IOnlineSubsystem;
 
@@ -246,3 +248,7 @@ private:
 /** Public references to the online subsystem module pointer should use this */
 typedef TSharedPtr<FOnlineSubsystemModule, ESPMode::ThreadSafe> FOnlineSubsystemModulePtr;
 
+
+#if UE_ENABLE_INCLUDE_ORDER_DEPRECATED_IN_5_2
+#include "CoreMinimal.h"
+#endif
diff --git a/Engine/Plugins/Online/OnlineSubsystem/Source/Public/OnlineSubsystemNames.h b/Engine/Plugins/Online/OnlineSubsystem/Source/Public/OnlineSubsystemNames.h
--- a/Engine/Plugins/Online/OnlineSubsystem/Source/Public/OnlineSubsystemNames.h
+++ b/Engine/Plugins/Online/OnlineSubsystem/Source/Public/OnlineSubsystemNames.h
@@ -2,7 +2,7 @@
 
 #pragma once
 
-#include "CoreMinimal.h"
+#include "UObject/NameTypes.h"
 
 
 #define OSS_PLATFORM_NAME_PS4		TEXT("PSN")
@@ -98,3 +98,7 @@ UE_DEPRECATED(5.2, "OnlineSubsystemOculus has been deprecated and will be remove
 #ifndef EOS_SUBSYSTEM
 #define EOS_SUBSYSTEM FName(TEXT("EOS"))
 #endif
+
+#if UE_ENABLE_INCLUDE_ORDER_DEPRECATED_IN_5_2
+#include "CoreMinimal.h"
+#endif
diff --git a/Engine/Plugins/Online/OnlineSubsystem/Source/Public/OnlineSubsystemTypes.h b/Engine/Plugins/Online/OnlineSubsystem/Source/Public/OnlineSubsystemTypes.h
--- a/Engine/Plugins/Online/OnlineSubsystem/Source/Public/OnlineSubsystemTypes.h
+++ b/Engine/Plugins/Online/OnlineSubsystem/Source/Public/OnlineSubsystemTypes.h
@@ -2,11 +2,8 @@
 
 #pragma once
 
-#include "CoreMinimal.h"
 #include "Online/CoreOnline.h"
-#include "Online/OnlineBase.h"
-#include "OnlineSubsystemNames.h"  // can be removed once we have no more temporary FUniqueNetId subtypes
-#include "OnlineSubsystemPackage.h"
+#include "OnlineSubsystemNames.h" // IWYU pragma: keep
 
 
 #if UE_GAME && UE_BUILD_SHIPPING
@@ -1642,3 +1639,10 @@ inline void ParseOnlineSubsystemConfigPairs(TArrayView<const FString> InEntries,
 		OutPairs.Emplace(MoveTemp(KeyString), MoveTemp(ValueString));
 	}
 }
+
+#if UE_ENABLE_INCLUDE_ORDER_DEPRECATED_IN_5_2
+#include "CoreMinimal.h"
+#include "Online/OnlineBase.h"
+#include "OnlineSubsystemNames.h"  // can be removed once we have no more temporary FUniqueNetId subtypes
+#include "OnlineSubsystemPackage.h"
+#endif
diff --git a/Engine/Plugins/Online/OnlineSubsystem/Source/Test/Utils/OnlineTestCommon.h b/Engine/Plugins/Online/OnlineSubsystem/Source/Test/Utils/OnlineTestCommon.h
--- a/Engine/Plugins/Online/OnlineSubsystem/Source/Test/Utils/OnlineTestCommon.h
+++ b/Engine/Plugins/Online/OnlineSubsystem/Source/Test/Utils/OnlineTestCommon.h
@@ -1,6 +1,5 @@
 // Copyright Epic Games, Inc. All Rights Reserved.
 #pragma once
-#include "Online.h"
 #include "OnlineSubsystem.h"
 #include "Interfaces/OnlineIdentityInterface.h"
 #include "Interfaces/OnlineFriendsInterface.h"
@@ -183,4 +182,8 @@ public:
 	FDelegateHandle OnBlockedPlayerCompleteDelegateHandle;
 	FDelegateHandle OnUnblockedPlayerCompleteDelegateHandle;
 	FDelegateHandle OnSendMessageCompleteDelegateHandle;
-};
\ No newline at end of file
+};
+
+#if UE_ENABLE_INCLUDE_ORDER_DEPRECATED_IN_5_2
+#include "Online.h"
+#endif
diff --git a/Engine/Plugins/Online/OnlineSubsystemEOS/Source/OnlineSubsystemEOS/Private/OnlineSessionEOS.cpp b/Engine/Plugins/Online/OnlineSubsystemEOS/Source/OnlineSubsystemEOS/Private/OnlineSessionEOS.cpp
--- a/Engine/Plugins/Online/OnlineSubsystemEOS/Source/OnlineSubsystemEOS/Private/OnlineSessionEOS.cpp
+++ b/Engine/Plugins/Online/OnlineSubsystemEOS/Source/OnlineSubsystemEOS/Private/OnlineSessionEOS.cpp
@@ -1,7 +1,12 @@
 // Copyright Epic Games, Inc. All Rights Reserved.
 
 #include "OnlineSessionEOS.h"
+
+#if WITH_EOS_SDK
+
 #include "Misc/Guid.h"
+#include "Online/OnlineBase.h"
+#include "Online/OnlineSessionNames.h"
 #include "OnlineSubsystem.h"
 #include "OnlineSubsystemEOS.h"
 #include "OnlineSubsystemEOSPrivate.h"
@@ -16,10 +21,9 @@
 #include "NetDriverEOS.h"
 #include "EOSVoiceChatUser.h"
 
-#if WITH_EOS_SDK
-	#include "eos_sessions.h"
-	#include "eos_metrics.h"
-	#include "eos_lobby.h"
+#include "eos_sessions.h"
+#include "eos_metrics.h"
+#include "eos_lobby.h"
 
 /** This is the game name plus version in ansi done once for optimization */
 char BucketIdAnsi[EOS_OSS_BUCKET_ID_STRING_LENGTH + 1]; // X characters plus null terminator
@@ -4517,4 +4521,4 @@ void FOnlineSessionEOS::CopyLobbyMemberAttributes(const FLobbyDetailsEOS& LobbyD
 	}
 }
 
-#endif
+#endif // WITH_EOS_SDK
diff --git a/Engine/Plugins/Online/OnlineSubsystemEOS/Source/OnlineSubsystemEOS/Private/OnlineSubsystemEOSTypes.h b/Engine/Plugins/Online/OnlineSubsystemEOS/Source/OnlineSubsystemEOS/Private/OnlineSubsystemEOSTypes.h
--- a/Engine/Plugins/Online/OnlineSubsystemEOS/Source/OnlineSubsystemEOS/Private/OnlineSubsystemEOSTypes.h
+++ b/Engine/Plugins/Online/OnlineSubsystemEOS/Source/OnlineSubsystemEOS/Private/OnlineSubsystemEOSTypes.h
@@ -11,6 +11,7 @@
 #include "IPAddress.h"
 #include "OnlineSubsystem.h"
 #include "OnlineSubsystemTypes.h"
+#include "OnlineSubsystemEOSPackage.h" // IWYU pragma: keep
 
 #define EOS_OSS_BUCKET_ID_STRING_LENGTH 60
 #define EOS_OSS_STRING_BUFFER_LENGTH 256 + 1 // 256 plus null terminator
diff --git a/Engine/Plugins/Online/OnlineSubsystemEOS/Source/OnlineSubsystemEOSPlus/Private/OnlineUserEOSPlus.h b/Engine/Plugins/Online/OnlineSubsystemEOS/Source/OnlineSubsystemEOSPlus/Private/OnlineUserEOSPlus.h
--- a/Engine/Plugins/Online/OnlineSubsystemEOS/Source/OnlineSubsystemEOSPlus/Private/OnlineUserEOSPlus.h
+++ b/Engine/Plugins/Online/OnlineSubsystemEOS/Source/OnlineSubsystemEOSPlus/Private/OnlineUserEOSPlus.h
@@ -9,6 +9,7 @@
 #include "Interfaces/OnlineIdentityInterface.h"
 #include "Interfaces/OnlineFriendsInterface.h"
 #include "Interfaces/OnlinePresenceInterface.h"
+#include "OnlineSubsystemEOSPlusPackage.h" // IWYU pragma: keep
 
 class FOnlineSubsystemEOSPlus;
 
diff --git a/Engine/Plugins/Online/OnlineSubsystemTencent/Source/Private/OnlineAsyncTaskManagerTencent.h b/Engine/Plugins/Online/OnlineSubsystemTencent/Source/Private/OnlineAsyncTaskManagerTencent.h
--- a/Engine/Plugins/Online/OnlineSubsystemTencent/Source/Private/OnlineAsyncTaskManagerTencent.h
+++ b/Engine/Plugins/Online/OnlineSubsystemTencent/Source/Private/OnlineAsyncTaskManagerTencent.h
@@ -7,6 +7,7 @@
 #include "OnlineAsyncTaskManager.h"
 
 #if WITH_TENCENT_RAIL_SDK
+#include "Containers/Set.h"
 #include "RailSDK.h"
 #endif
 
diff --git a/Engine/Plugins/Online/OnlineSubsystemTencent/Source/Private/OnlineFriendsTencent.cpp b/Engine/Plugins/Online/OnlineSubsystemTencent/Source/Private/OnlineFriendsTencent.cpp
--- a/Engine/Plugins/Online/OnlineSubsystemTencent/Source/Private/OnlineFriendsTencent.cpp
+++ b/Engine/Plugins/Online/OnlineSubsystemTencent/Source/Private/OnlineFriendsTencent.cpp
@@ -7,6 +7,7 @@
 
 #include "Interfaces/OnlineIdentityInterface.h"
 #include "Interfaces/OnlineUserInterface.h"
+#include "Online/OnlineBase.h"
 #include "OnlineIdentityTencent.h"
 #include "OnlinePresenceTencent.h"
 #include "OnlineSubsystemTencent.h"
diff --git a/Engine/Plugins/Online/OnlineSubsystemTencent/Source/Private/OnlinePlayTimeLimitTencent.h b/Engine/Plugins/Online/OnlineSubsystemTencent/Source/Private/OnlinePlayTimeLimitTencent.h
--- a/Engine/Plugins/Online/OnlineSubsystemTencent/Source/Private/OnlinePlayTimeLimitTencent.h
+++ b/Engine/Plugins/Online/OnlineSubsystemTencent/Source/Private/OnlinePlayTimeLimitTencent.h
@@ -7,6 +7,7 @@
 #if WITH_TENCENTSDK
 #if WITH_TENCENT_RAIL_SDK
 
+#include "Delegates/IDelegateInstance.h"
 #include "PlayTimeLimitUser.h"
 
 class ONLINESUBSYSTEMTENCENT_API FOnlinePlayTimeLimitUserTencentRail
diff --git a/Engine/Plugins/Online/OnlineSubsystemTencent/Source/Private/OnlineSessionTencent.cpp b/Engine/Plugins/Online/OnlineSubsystemTencent/Source/Private/OnlineSessionTencent.cpp
--- a/Engine/Plugins/Online/OnlineSubsystemTencent/Source/Private/OnlineSessionTencent.cpp
+++ b/Engine/Plugins/Online/OnlineSubsystemTencent/Source/Private/OnlineSessionTencent.cpp
@@ -2,7 +2,6 @@
 
 #include "OnlineSessionTencent.h"
 #include "OnlineSubsystemTencent.h"
-#include "OnlineSubsystemTencentTypes.h"
 #include "OnlineIdentityTencent.h"
 #include "Misc/ScopeLock.h"
 
diff --git a/Engine/Plugins/Online/OnlineSubsystemTencent/Source/Private/OnlineSessionTencentRail.cpp b/Engine/Plugins/Online/OnlineSubsystemTencent/Source/Private/OnlineSessionTencentRail.cpp
--- a/Engine/Plugins/Online/OnlineSubsystemTencent/Source/Private/OnlineSessionTencentRail.cpp
+++ b/Engine/Plugins/Online/OnlineSubsystemTencent/Source/Private/OnlineSessionTencentRail.cpp
@@ -4,6 +4,7 @@
 
 #if WITH_TENCENT_RAIL_SDK
 #include "Interfaces/OnlineIdentityInterface.h"
+#include "Online/OnlineBase.h"
 #include "OnlineAsyncTasksTencent.h"
 #include "OnlinePresenceTencent.h"
 #include "OnlineSubsystemTencent.h"
diff --git a/Engine/Plugins/Online/OnlineSubsystemTencent/Source/Private/OnlineSubsystemTencent.cpp b/Engine/Plugins/Online/OnlineSubsystemTencent/Source/Private/OnlineSubsystemTencent.cpp
--- a/Engine/Plugins/Online/OnlineSubsystemTencent/Source/Private/OnlineSubsystemTencent.cpp
+++ b/Engine/Plugins/Online/OnlineSubsystemTencent/Source/Private/OnlineSubsystemTencent.cpp
@@ -15,9 +15,7 @@
 #include "Interfaces/OnlineTitleFileInterface.h"
 #include "Interfaces/OnlineUserCloudInterface.h"
 #include "Interfaces/VoiceInterface.h"
-#include "Internationalization/Internationalization.h"
 #include "OnlineSessionTencent.h"
-#include "OnlineSubsystemNames.h"
 #include "Trace/Trace.inl"
 
 #if WITH_TENCENT_RAIL_SDK
diff --git a/Engine/Plugins/Online/OnlineSubsystemTencent/Source/Public/OnlineSessionTencentRail.h b/Engine/Plugins/Online/OnlineSubsystemTencent/Source/Public/OnlineSessionTencentRail.h
--- a/Engine/Plugins/Online/OnlineSubsystemTencent/Source/Public/OnlineSessionTencentRail.h
+++ b/Engine/Plugins/Online/OnlineSubsystemTencent/Source/Public/OnlineSessionTencentRail.h
@@ -2,7 +2,6 @@
 
 #pragma once
 
-#include "HAL/Platform.h"
 
 #if WITH_TENCENT_RAIL_SDK
 
@@ -231,3 +230,7 @@ typedef TSharedPtr<FOnlineSessionTencentRail, ESPMode::ThreadSafe> FOnlineSessio
 
 
 #endif // WITH_TENCENT_RAIL_SDK
+
+#if UE_ENABLE_INCLUDE_ORDER_DEPRECATED_IN_5_2
+#include "HAL/Platform.h"
+#endif
diff --git a/Engine/Plugins/Online/OnlineSubsystemTencent/Source/Public/OnlineSubsystemTencent.h b/Engine/Plugins/Online/OnlineSubsystemTencent/Source/Public/OnlineSubsystemTencent.h
--- a/Engine/Plugins/Online/OnlineSubsystemTencent/Source/Public/OnlineSubsystemTencent.h
+++ b/Engine/Plugins/Online/OnlineSubsystemTencent/Source/Public/OnlineSubsystemTencent.h
@@ -5,8 +5,8 @@
 #if WITH_TENCENTSDK
 
 #include "OnlineSubsystemImpl.h"
-#include "OnlineSubsystemTencentPackage.h"
-#include "OnlineSubsystemTencentTypes.h"
+#include "OnlineSubsystemTencentTypes.h" // IWYU pragma: keep
+#include "OnlineSubsystemPackage.h"
 
 /**
  * Delegate called when Tencent requires Anti-Addiction message to be displayed
@@ -202,3 +202,8 @@ private:
 typedef TSharedPtr<FOnlineSubsystemTencent, ESPMode::ThreadSafe> FOnlineSubsystemTencentPtr;
 
 #endif // WITH_TENCENTSDK
+
+#if UE_ENABLE_INCLUDE_ORDER_DEPRECATED_IN_5_2
+#include "OnlineSubsystemTencentPackage.h"
+#include "OnlineSubsystemTencentTypes.h"
+#endif
diff --git a/Engine/Plugins/Online/OnlineSubsystemTencent/Source/Public/OnlineSubsystemTencentTypes.h b/Engine/Plugins/Online/OnlineSubsystemTencent/Source/Public/OnlineSubsystemTencentTypes.h
--- a/Engine/Plugins/Online/OnlineSubsystemTencent/Source/Public/OnlineSubsystemTencentTypes.h
+++ b/Engine/Plugins/Online/OnlineSubsystemTencent/Source/Public/OnlineSubsystemTencentTypes.h
@@ -2,10 +2,9 @@
 
 #pragma once
 
-#include "OnlineSubsystemPackage.h"
 #include "OnlineSubsystemTypes.h"
 #include "Misc/Guid.h"
-#include "OnlineSubsystemNames.h"
+#include "OnlineSubsystemTencentPackage.h"
 
 #if WITH_TENCENTSDK
 
diff --git a/Engine/Plugins/Runtime/AndroidPermission/Source/AndroidPermission/Private/AndroidPermissionCallbackProxy.cpp b/Engine/Plugins/Runtime/AndroidPermission/Source/AndroidPermission/Private/AndroidPermissionCallbackProxy.cpp
--- a/Engine/Plugins/Runtime/AndroidPermission/Source/AndroidPermission/Private/AndroidPermissionCallbackProxy.cpp
+++ b/Engine/Plugins/Runtime/AndroidPermission/Source/AndroidPermission/Private/AndroidPermissionCallbackProxy.cpp
@@ -3,7 +3,7 @@
 #include "AndroidPermissionCallbackProxy.h"
 #include "AndroidPermission.h"
 
-#include "Async/TaskGraphInterfaces.h"
+#include "Trace/Trace.inl"
 
 #include UE_INLINE_GENERATED_CPP_BY_NAME(AndroidPermissionCallbackProxy)
 
diff --git a/Engine/Plugins/Runtime/AudioWidgets/Source/AudioWidgets/Public/AudioMeter.h b/Engine/Plugins/Runtime/AudioWidgets/Source/AudioWidgets/Public/AudioMeter.h
--- a/Engine/Plugins/Runtime/AudioWidgets/Source/AudioWidgets/Public/AudioMeter.h
+++ b/Engine/Plugins/Runtime/AudioWidgets/Source/AudioWidgets/Public/AudioMeter.h
@@ -150,13 +150,15 @@ protected:
 	PROPERTY_BINDING_IMPLEMENTATION(TArray<FMeterChannelInfo>, MeterChannelInfo);
 };
 
-namespace AudioWidgets
+#if UE_ENABLE_INCLUDE_ORDER_DEPRECATED_IN_5_2
+#include "CoreMinimal.h"
+#include "Styling/SlateTypes.h"
 {
 	class AUDIOWIDGETS_API FAudioMeter
 	{
 	public:
 		FAudioMeter();
-
+#endif
 		UAudioBus* GetAudioBus() const;
 
 		TSharedRef<SAudioMeter> GetWidget() const;
diff --git a/Engine/Plugins/Runtime/AudioWidgets/Source/AudioWidgets/Public/AudioMeterStyle.h b/Engine/Plugins/Runtime/AudioWidgets/Source/AudioWidgets/Public/AudioMeterStyle.h
--- a/Engine/Plugins/Runtime/AudioWidgets/Source/AudioWidgets/Public/AudioMeterStyle.h
+++ b/Engine/Plugins/Runtime/AudioWidgets/Source/AudioWidgets/Public/AudioMeterStyle.h
@@ -2,13 +2,13 @@
 
 #pragma once
 
-#include "CoreMinimal.h"
-#include "UObject/ObjectMacros.h"
-#include "Widgets/SWidget.h"
-#include "Components/Widget.h"
+#include "Fonts/SlateFontInfo.h"
 #include "Styling/SlateWidgetStyle.h"
+#include "Styling/SlateBrush.h"
 #include "AudioMeterStyle.generated.h"
 
+struct FCompositeFont;
+
 /**
  * Represents the appearance of an SAudioMeter
  */
@@ -139,3 +139,9 @@ struct AUDIOWIDGETS_API FAudioMeterStyle : public FSlateWidgetStyle
 };
 
 
+
+#if UE_ENABLE_INCLUDE_ORDER_DEPRECATED_IN_5_2
+#include "Components/Widget.h"
+#include "CoreMinimal.h"
+#include "Widgets/SWidget.h"
+#endif
diff --git a/Engine/Plugins/Runtime/AudioWidgets/Source/AudioWidgets/Public/AudioRadialSlider.h b/Engine/Plugins/Runtime/AudioWidgets/Source/AudioWidgets/Public/AudioRadialSlider.h
--- a/Engine/Plugins/Runtime/AudioWidgets/Source/AudioWidgets/Public/AudioRadialSlider.h
+++ b/Engine/Plugins/Runtime/AudioWidgets/Source/AudioWidgets/Public/AudioRadialSlider.h
@@ -2,14 +2,13 @@
 
 #pragma once
 
-#include "AudioWidgetsStyle.h"
+#include "AudioWidgetsSlateTypes.h"
 #include "Components/Widget.h"
-#include "SAudioRadialSlider.h"
-#include "Styling/StyleColors.h"
-#include "UObject/ObjectMacros.h"
-#include "Widgets/SWidget.h"
 #include "AudioRadialSlider.generated.h"
 
+class SAudioRadialSlider;
+enum EAudioRadialSliderLayout : int;
+
 DECLARE_DYNAMIC_MULTICAST_DELEGATE_OneParam(FOnAudioRadialSliderValueChangedEvent, float, Value);
 
 /**
@@ -196,3 +195,10 @@ class AUDIOWIDGETS_API UAudioFrequencyRadialSlider : public UAudioRadialSlider
 protected:
 	virtual TSharedRef<SWidget> RebuildWidget() override;
 };
+
+#if UE_ENABLE_INCLUDE_ORDER_DEPRECATED_IN_5_2
+#include "AudioWidgetsStyle.h"
+#include "SAudioRadialSlider.h"
+#include "Styling/StyleColors.h"
+#include "Widgets/SWidget.h"
+#endif
diff --git a/Engine/Plugins/Runtime/AudioWidgets/Source/AudioWidgets/Public/AudioSlider.h b/Engine/Plugins/Runtime/AudioWidgets/Source/AudioWidgets/Public/AudioSlider.h
--- a/Engine/Plugins/Runtime/AudioWidgets/Source/AudioWidgets/Public/AudioSlider.h
+++ b/Engine/Plugins/Runtime/AudioWidgets/Source/AudioWidgets/Public/AudioSlider.h
@@ -5,14 +5,10 @@
 #include "AudioDefines.h"
 #include "AudioWidgetsSlateTypes.h"
 #include "Components/Slider.h"
-#include "Components/Widget.h"
-#include "Curves/CurveFloat.h"
-#include "Styling/StyleColors.h"
-#include "Styling/SlateTypes.h"
-#include "UObject/ObjectMacros.h"
-#include "Widgets/SWidget.h"
 #include "AudioSlider.generated.h"
 
+class UCurveFloat;
+
 class SAudioSliderBase;
 
 /**
@@ -230,3 +226,9 @@ class AUDIOWIDGETS_API UAudioFrequencySlider : public UAudioSliderBase
 protected:
 	virtual TSharedRef<SWidget> RebuildWidget() override;
 };
+
+#if UE_ENABLE_INCLUDE_ORDER_DEPRECATED_IN_5_2
+#include "Curves/CurveFloat.h"
+#include "Styling/StyleColors.h"
+#include "Widgets/SWidget.h"
+#endif
diff --git a/Engine/Plugins/Runtime/AudioWidgets/Source/AudioWidgets/Public/AudioWidgets.h b/Engine/Plugins/Runtime/AudioWidgets/Source/AudioWidgets/Public/AudioWidgets.h
--- a/Engine/Plugins/Runtime/AudioWidgets/Source/AudioWidgets/Public/AudioWidgets.h
+++ b/Engine/Plugins/Runtime/AudioWidgets/Source/AudioWidgets/Public/AudioWidgets.h
@@ -2,8 +2,7 @@
 
 #pragma once
 
-#include "AudioWidgetsStyle.h"
-#include "Modules/ModuleManager.h"
+#include "Modules/ModuleInterface.h"
 
 class FAudioWidgetsModule : public IModuleInterface
 {
@@ -13,3 +12,8 @@ public:
 	virtual void StartupModule() override;
 	virtual void ShutdownModule() override;
 };
+
+#if UE_ENABLE_INCLUDE_ORDER_DEPRECATED_IN_5_2
+#include "AudioWidgetsStyle.h"
+#include "Modules/ModuleManager.h"
+#endif
diff --git a/Engine/Plugins/Runtime/AudioWidgets/Source/AudioWidgets/Public/AudioWidgetsSlateTypes.h b/Engine/Plugins/Runtime/AudioWidgets/Source/AudioWidgets/Public/AudioWidgetsSlateTypes.h
--- a/Engine/Plugins/Runtime/AudioWidgets/Source/AudioWidgets/Public/AudioWidgetsSlateTypes.h
+++ b/Engine/Plugins/Runtime/AudioWidgets/Source/AudioWidgets/Public/AudioWidgetsSlateTypes.h
@@ -2,12 +2,7 @@
 
 #pragma once
 
-#include "Styling/StyleDefaults.h"
-#include "Styling/SlateStyle.h"
-#include "Styling/SlateStyleMacros.h"
 #include "Styling/SlateTypes.h"
-#include "Styling/SlateWidgetStyle.h"
-#include "UObject/ObjectMacros.h"
 
 #include "AudioWidgetsSlateTypes.generated.h"
 
@@ -174,3 +169,9 @@ struct AUDIOWIDGETS_API FAudioRadialSliderStyle : public FSlateWidgetStyle
 	float DefaultSliderRadius;
 	FAudioRadialSliderStyle& SetDefaultSliderRadius(const float& InDefaultSliderRadius) { DefaultSliderRadius = InDefaultSliderRadius; return *this; }
 };
+
+#if UE_ENABLE_INCLUDE_ORDER_DEPRECATED_IN_5_2
+#include "Styling/SlateStyle.h"
+#include "Styling/SlateStyleMacros.h"
+#include "Styling/StyleDefaults.h"
+#endif
diff --git a/Engine/Plugins/Runtime/AudioWidgets/Source/AudioWidgets/Public/AudioWidgetsStyle.h b/Engine/Plugins/Runtime/AudioWidgets/Source/AudioWidgets/Public/AudioWidgetsStyle.h
--- a/Engine/Plugins/Runtime/AudioWidgets/Source/AudioWidgets/Public/AudioWidgetsStyle.h
+++ b/Engine/Plugins/Runtime/AudioWidgets/Source/AudioWidgets/Public/AudioWidgetsStyle.h
@@ -3,9 +3,6 @@
 #pragma once
 
 #include "Styling/SlateStyle.h"
-#include "Styling/SlateStyleMacros.h"
-#include "Styling/SlateStyleRegistry.h"
-#include "Styling/SlateTypes.h"
 
 /**
  * Slate style set that defines all the styles for audio widgets
@@ -25,3 +22,9 @@ private:
 	FAudioWidgetsStyle();
 	~FAudioWidgetsStyle();
 };
+
+#if UE_ENABLE_INCLUDE_ORDER_DEPRECATED_IN_5_2
+#include "Styling/SlateStyleMacros.h"
+#include "Styling/SlateStyleRegistry.h"
+#include "Styling/SlateTypes.h"
+#endif
diff --git a/Engine/Plugins/Runtime/AudioWidgets/Source/AudioWidgets/Public/SAudioMeter.h b/Engine/Plugins/Runtime/AudioWidgets/Source/AudioWidgets/Public/SAudioMeter.h
--- a/Engine/Plugins/Runtime/AudioWidgets/Source/AudioWidgets/Public/SAudioMeter.h
+++ b/Engine/Plugins/Runtime/AudioWidgets/Source/AudioWidgets/Public/SAudioMeter.h
@@ -2,15 +2,9 @@
 
 #pragma once
 
-#include "CoreMinimal.h"
-#include "Misc/Attribute.h"
-#include "Input/Reply.h"
+#include "Styling/ISlateStyle.h"
 #include "Styling/SlateWidgetStyleAsset.h"
 #include "Widgets/DeclarativeSyntaxSupport.h"
-#include "Styling/SlateColor.h"
-#include "Styling/SlateTypes.h"
-#include "Styling/CoreStyle.h"
-#include "Framework/SlateDelegates.h"
 #include "Widgets/SLeafWidget.h"
 #include "AudioMeterStyle.h"
 #include "AudioMeterTypes.h"
@@ -129,3 +123,10 @@ protected:
 
 	TAttribute<TArray<FMeterChannelInfo>> MeterChannelInfoAttribute;
 };
+
+#if UE_ENABLE_INCLUDE_ORDER_DEPRECATED_IN_5_2
+#include "CoreMinimal.h"
+#include "Framework/SlateDelegates.h"
+#include "Styling/CoreStyle.h"
+#include "Styling/SlateTypes.h"
+#endif
diff --git a/Engine/Plugins/Runtime/AudioWidgets/Source/AudioWidgets/Public/SAudioRadialSlider.h b/Engine/Plugins/Runtime/AudioWidgets/Source/AudioWidgets/Public/SAudioRadialSlider.h
--- a/Engine/Plugins/Runtime/AudioWidgets/Source/AudioWidgets/Public/SAudioRadialSlider.h
+++ b/Engine/Plugins/Runtime/AudioWidgets/Source/AudioWidgets/Public/SAudioRadialSlider.h
@@ -2,21 +2,17 @@
 
 #pragma once
 
-#include "AudioDefines.h"
 #include "AudioWidgetsSlateTypes.h"
 #include "Curves/CurveFloat.h"
 #include "Framework/SlateDelegates.h"
 #include "SAudioInputWidget.h"
-#include "SAudioTextBox.h"
-#include "Styling/CoreStyle.h"
+#include "Styling/ISlateStyle.h"
 #include "Styling/SlateStyleRegistry.h"
 #include "Styling/SlateWidgetStyleAsset.h"
-#include "Widgets/DeclarativeSyntaxSupport.h"
-#include "Widgets/SCompoundWidget.h"
-#include "Widgets/SOverlay.h"
-#include "Widgets/Images/SImage.h"
-#include "Widgets/Input/SEditableText.h"
-#include "Widgets/Layout/SWidgetSwitcher.h"
+
+class SAudioTextBox;
+class SImage;
+class SWidgetSwitcher;
 
 #include "SAudioRadialSlider.generated.h"
 
@@ -199,4 +195,14 @@ public:
 	void Construct(const SAudioRadialSlider::FArguments& InArgs);
 	const float GetOutputValue(const float InSliderValue);
 	const float GetSliderValue(const float OutputValue);
-};
\ No newline at end of file
+};
+
+#if UE_ENABLE_INCLUDE_ORDER_DEPRECATED_IN_5_2
+#include "AudioDefines.h"
+#include "SAudioTextBox.h"
+#include "Styling/CoreStyle.h"
+#include "Widgets/Images/SImage.h"
+#include "Widgets/Input/SEditableText.h"
+#include "Widgets/Layout/SWidgetSwitcher.h"
+#include "Widgets/SOverlay.h"
+#endif
diff --git a/Engine/Plugins/Runtime/AudioWidgets/Source/AudioWidgets/Public/SAudioSlider.h b/Engine/Plugins/Runtime/AudioWidgets/Source/AudioWidgets/Public/SAudioSlider.h
--- a/Engine/Plugins/Runtime/AudioWidgets/Source/AudioWidgets/Public/SAudioSlider.h
+++ b/Engine/Plugins/Runtime/AudioWidgets/Source/AudioWidgets/Public/SAudioSlider.h
@@ -3,20 +3,17 @@
 #pragma once
 
 #include "AudioWidgetsSlateTypes.h"
-#include "Curves/CurveFloat.h"
 #include "Framework/SlateDelegates.h"
-#include "Misc/Attribute.h"
 #include "SAudioInputWidget.h"
-#include "SAudioTextBox.h"
+#include "Styling/ISlateStyle.h"
 #include "Styling/SlateStyleRegistry.h"
-#include "Widgets/DeclarativeSyntaxSupport.h"
-#include "Widgets/SCompoundWidget.h"
-#include "Widgets/SOverlay.h"
-#include "Widgets/Images/SImage.h"
-#include "Widgets/Input/SEditableText.h"
-#include "Widgets/Input/SSlider.h"
-#include "Widgets/Layout/SConstraintCanvas.h"
-#include "Widgets/Layout/SWidgetSwitcher.h"
+#include "Styling/SlateWidgetStyleAsset.h"
+
+class SAudioTextBox;
+class SImage;
+class SSlider;
+class SWidgetSwitcher;
+class UCurveFloat;
 
 /**
  * Slate audio sliders that wrap SSlider and provides additional audio specific functionality.
@@ -217,3 +214,14 @@ public:
 	const float GetOutputValue(const float InSliderValue);
 	const float GetSliderValue(const float OutputValue);
 };
+
+#if UE_ENABLE_INCLUDE_ORDER_DEPRECATED_IN_5_2
+#include "Curves/CurveFloat.h"
+#include "SAudioTextBox.h"
+#include "Widgets/Images/SImage.h"
+#include "Widgets/Input/SEditableText.h"
+#include "Widgets/Input/SSlider.h"
+#include "Widgets/Layout/SConstraintCanvas.h"
+#include "Widgets/Layout/SWidgetSwitcher.h"
+#include "Widgets/SOverlay.h"
+#endif
diff --git a/Engine/Plugins/Runtime/AudioWidgets/Source/AudioWidgets/Public/SAudioTextBox.h b/Engine/Plugins/Runtime/AudioWidgets/Source/AudioWidgets/Public/SAudioTextBox.h
--- a/Engine/Plugins/Runtime/AudioWidgets/Source/AudioWidgets/Public/SAudioTextBox.h
+++ b/Engine/Plugins/Runtime/AudioWidgets/Source/AudioWidgets/Public/SAudioTextBox.h
@@ -3,19 +3,17 @@
 #pragma once
 
 #include "AudioWidgetsSlateTypes.h"
-#include "Fonts/FontMeasure.h"
-#include "Framework/Application/SlateApplication.h"
-#include "Styling/CoreStyle.h"
+#include "Framework/SlateDelegates.h"
 #include "Styling/SlateStyleRegistry.h"
+#include "Styling/ISlateStyle.h"
 #include "Styling/SlateWidgetStyleAsset.h"
-#include "Widgets/DeclarativeSyntaxSupport.h"
 #include "Widgets/Input/NumericTypeInterface.h"
-#include "Widgets/Images/SImage.h"
-#include "Widgets/Input/SEditableText.h"
-#include "Widgets/Layout/SBorder.h"
 #include "Widgets/SBoxPanel.h"
 #include "Widgets/SCompoundWidget.h"
 
+class SBorder;
+class SEditableText;
+
 // FVariablePrecisionNumericInterface
 // Taken from PropertyEditor/VariablePrecisionNumericInterface.h
 // Allow more precision as the numbers get closer to zero
@@ -104,3 +102,12 @@ protected:
 	/** Used to convert and format value text strings **/
 	static const FVariablePrecisionNumericInterface NumericInterface;
 };
+
+#if UE_ENABLE_INCLUDE_ORDER_DEPRECATED_IN_5_2
+#include "Fonts/FontMeasure.h"
+#include "Framework/Application/SlateApplication.h"
+#include "Styling/CoreStyle.h"
+#include "Widgets/Images/SImage.h"
+#include "Widgets/Input/SEditableText.h"
+#include "Widgets/Layout/SBorder.h"
+#endif
diff --git a/Engine/Plugins/Runtime/GameplayBehaviorSmartObjects/Source/GameplayBehaviorSmartObjectsModule/Public/AI/BTTask_FindAndUseGameplayBehaviorSmartObject.h b/Engine/Plugins/Runtime/GameplayBehaviorSmartObjects/Source/GameplayBehaviorSmartObjectsModule/Public/AI/BTTask_FindAndUseGameplayBehaviorSmartObject.h
--- a/Engine/Plugins/Runtime/GameplayBehaviorSmartObjects/Source/GameplayBehaviorSmartObjectsModule/Public/AI/BTTask_FindAndUseGameplayBehaviorSmartObject.h
+++ b/Engine/Plugins/Runtime/GameplayBehaviorSmartObjects/Source/GameplayBehaviorSmartObjectsModule/Public/AI/BTTask_FindAndUseGameplayBehaviorSmartObject.h
@@ -3,10 +3,11 @@
 #pragma once
 
 #include "BehaviorTree/BTTaskNode.h"
-#include "AI/AITask_UseGameplayBehaviorSmartObject.h"
 #include "EnvironmentQuery/EnvQueryTypes.h"
 #include "BTTask_FindAndUseGameplayBehaviorSmartObject.generated.h"
 
+class UAITask_UseGameplayBehaviorSmartObject;
+
 
 class AITask_UseSmartObject;
 
@@ -56,3 +57,7 @@ protected:
 
 	FQueryFinishedSignature EQSQueryFinishedDelegate; 
 };
+
+#if UE_ENABLE_INCLUDE_ORDER_DEPRECATED_IN_5_2
+#include "AI/AITask_UseGameplayBehaviorSmartObject.h"
+#endif
diff --git a/Engine/Plugins/Runtime/GameplayStateTree/Source/GameplayStateTreeModule/Public/Components/StateTreeComponent.h b/Engine/Plugins/Runtime/GameplayStateTree/Source/GameplayStateTreeModule/Public/Components/StateTreeComponent.h
--- a/Engine/Plugins/Runtime/GameplayStateTree/Source/GameplayStateTreeModule/Public/Components/StateTreeComponent.h
+++ b/Engine/Plugins/Runtime/GameplayStateTree/Source/GameplayStateTreeModule/Public/Components/StateTreeComponent.h
@@ -2,14 +2,18 @@
 
 #pragma once
 
-#include "GameplayTagContainer.h"
-#include "StateTreeExecutionContext.h"
 #include "BrainComponent.h"
+#include "GameplayTaskOwnerInterface.h"
 #include "StateTreeReference.h"
-#include "Tasks/AITask.h"
-#include "PropertyBag.h"
+#include "StateTreeInstanceData.h"
+#include "UObject/Package.h"
 #include "StateTreeComponent.generated.h"
 
+enum class EStateTreeRunStatus : uint8;
+struct FGameplayTag;
+struct FStateTreeEvent;
+struct FStateTreeExecutionContext;
+
 class UStateTree;
 
 DECLARE_DYNAMIC_MULTICAST_DELEGATE_OneParam(FStateTreeRunStatusChanged, EStateTreeRunStatus, StateTreeRunStatus);
@@ -100,3 +104,8 @@ protected:
 	/** if set, execution requests will be postponed */
 	uint8 bIsPaused : 1;
 };
+
+#if UE_ENABLE_INCLUDE_ORDER_DEPRECATED_IN_5_2
+#include "StateTreeExecutionContext.h"
+#include "Tasks/AITask.h"
+#endif
diff --git a/Engine/Plugins/Runtime/GameplayStateTree/Source/GameplayStateTreeModule/Public/Components/StateTreeComponentSchema.h b/Engine/Plugins/Runtime/GameplayStateTree/Source/GameplayStateTreeModule/Public/Components/StateTreeComponentSchema.h
--- a/Engine/Plugins/Runtime/GameplayStateTree/Source/GameplayStateTreeModule/Public/Components/StateTreeComponentSchema.h
+++ b/Engine/Plugins/Runtime/GameplayStateTree/Source/GameplayStateTreeModule/Public/Components/StateTreeComponentSchema.h
@@ -4,7 +4,7 @@
 
 #include "StateTreeSchema.h"
 #include "GameFramework/Actor.h"
-#include "Templates/SubclassOf.h"
+#include "StateTreeTypes.h"
 #include "StateTreeComponentSchema.generated.h"
 
 /**
diff --git a/Engine/Plugins/Runtime/GeometryCache/Source/GeometryCache/Classes/GeometryCache.h b/Engine/Plugins/Runtime/GeometryCache/Source/GeometryCache/Classes/GeometryCache.h
--- a/Engine/Plugins/Runtime/GeometryCache/Source/GeometryCache/Classes/GeometryCache.h
+++ b/Engine/Plugins/Runtime/GeometryCache/Source/GeometryCache/Classes/GeometryCache.h
@@ -1,9 +1,6 @@
 // Copyright Epic Games, Inc. All Rights Reserved.
 #pragma once
 
-#include "CoreMinimal.h"
-#include "UObject/ObjectMacros.h"
-#include "UObject/Object.h"
 #include "Interfaces/Interface_AssetUserData.h"
 #include "RenderCommandFence.h"
 
@@ -99,3 +96,7 @@ protected:
 	UPROPERTY()
 	uint64 Hash;
 };
+
+#if UE_ENABLE_INCLUDE_ORDER_DEPRECATED_IN_5_2
+#include "CoreMinimal.h"
+#endif
diff --git a/Engine/Plugins/Runtime/GeometryCache/Source/GeometryCache/Classes/GeometryCacheCodecBase.h b/Engine/Plugins/Runtime/GeometryCache/Source/GeometryCache/Classes/GeometryCacheCodecBase.h
--- a/Engine/Plugins/Runtime/GeometryCache/Source/GeometryCache/Classes/GeometryCacheCodecBase.h
+++ b/Engine/Plugins/Runtime/GeometryCache/Source/GeometryCache/Classes/GeometryCacheCodecBase.h
@@ -1,12 +1,12 @@
 // Copyright Epic Games, Inc. All Rights Reserved.
 #pragma once
 
-#include "CoreMinimal.h"
-#include "UObject/ObjectMacros.h"
 #include "Serialization/BulkData.h"
 
 #include "GeometryCacheCodecBase.generated.h"
 
+struct FResourceSizeEx;
+
 struct FGeometryCacheMeshData;
 class UGeometryCacheTrackStreamable;
 class FGeometryCacheTrackStreamableRenderResource;
@@ -150,3 +150,7 @@ protected:
 	UPROPERTY(VisibleAnywhere, Category = GeometryCache)
 	TArray<int32> TopologyRanges;
 };
+
+#if UE_ENABLE_INCLUDE_ORDER_DEPRECATED_IN_5_2
+#include "CoreMinimal.h"
+#endif
diff --git a/Engine/Plugins/Runtime/GeometryCache/Source/GeometryCache/Classes/GeometryCacheCodecRaw.h b/Engine/Plugins/Runtime/GeometryCache/Source/GeometryCache/Classes/GeometryCacheCodecRaw.h
--- a/Engine/Plugins/Runtime/GeometryCache/Source/GeometryCache/Classes/GeometryCacheCodecRaw.h
+++ b/Engine/Plugins/Runtime/GeometryCache/Source/GeometryCache/Classes/GeometryCacheCodecRaw.h
@@ -1,9 +1,6 @@
 // Copyright Epic Games, Inc. All Rights Reserved.
 #pragma once
 
-#include "CoreMinimal.h"
-#include "UObject/ObjectMacros.h"
-#include "UObject/ScriptMacros.h"
 #include "GeometryCacheCodecBase.h"
 
 #include "GeometryCacheCodecRaw.generated.h"
@@ -44,3 +41,7 @@ private:
 #endif
 };
 
+
+#if UE_ENABLE_INCLUDE_ORDER_DEPRECATED_IN_5_2
+#include "CoreMinimal.h"
+#endif
diff --git a/Engine/Plugins/Runtime/GeometryCache/Source/GeometryCache/Classes/GeometryCacheCodecV1.h b/Engine/Plugins/Runtime/GeometryCache/Source/GeometryCache/Classes/GeometryCacheCodecV1.h
--- a/Engine/Plugins/Runtime/GeometryCache/Source/GeometryCache/Classes/GeometryCacheCodecV1.h
+++ b/Engine/Plugins/Runtime/GeometryCache/Source/GeometryCache/Classes/GeometryCacheCodecV1.h
@@ -1,12 +1,7 @@
 // Copyright Epic Games, Inc. All Rights Reserved.
 #pragma once
 
-#include "CoreMinimal.h"
-#include "UObject/ObjectMacros.h"
-#include "UObject/ScriptMacros.h"
 #include "GeometryCacheCodecBase.h"
-#include "Serialization/MemoryWriter.h"
-#include "Serialization/BufferReader.h"
 
 #include "GeometryCacheCodecV1.generated.h"
 
@@ -56,3 +51,8 @@ private:
 #endif
 };
 
+
+#if UE_ENABLE_INCLUDE_ORDER_DEPRECATED_IN_5_2
+#include "CoreMinimal.h"
+#include "Serialization/MemoryWriter.h"
+#endif
diff --git a/Engine/Plugins/Runtime/GeometryCache/Source/GeometryCache/Classes/GeometryCacheMeshData.h b/Engine/Plugins/Runtime/GeometryCache/Source/GeometryCache/Classes/GeometryCacheMeshData.h
--- a/Engine/Plugins/Runtime/GeometryCache/Source/GeometryCache/Classes/GeometryCacheMeshData.h
+++ b/Engine/Plugins/Runtime/GeometryCache/Source/GeometryCache/Classes/GeometryCacheMeshData.h
@@ -1,11 +1,8 @@
 // Copyright Epic Games, Inc. All Rights Reserved.
 #pragma once
 
-#include "CoreMinimal.h"
-#include "UObject/ObjectMacros.h"
-#include "ProfilingDebugging/ResourceSize.h"
-#include "DynamicMeshBuilder.h"
-#include "UObject/GeometryObjectVersion.h"
+#include "Math/Box.h"
+#include "PackedNormal.h"
 #include "GeometryCacheMeshData.generated.h"
 
 /** Stores per-batch data used for rendering */
@@ -177,3 +174,9 @@ struct GEOMETRYCACHE_API FGeometryCacheMeshData
 private:
 	mutable uint64 Hash;
 };
+
+#if UE_ENABLE_INCLUDE_ORDER_DEPRECATED_IN_5_2
+#include "CoreMinimal.h"
+#include "DynamicMeshBuilder.h"
+#include "UObject/GeometryObjectVersion.h"
+#endif
diff --git a/Engine/Plugins/Runtime/GeometryCache/Source/GeometryCache/Classes/GeometryCacheTrack.h b/Engine/Plugins/Runtime/GeometryCache/Source/GeometryCache/Classes/GeometryCacheTrack.h
--- a/Engine/Plugins/Runtime/GeometryCache/Source/GeometryCache/Classes/GeometryCacheTrack.h
+++ b/Engine/Plugins/Runtime/GeometryCache/Source/GeometryCache/Classes/GeometryCacheTrack.h
@@ -1,11 +1,12 @@
 // Copyright Epic Games, Inc. All Rights Reserved.
 #pragma once
 
-#include "CoreMinimal.h"
-#include "UObject/ObjectMacros.h"
-#include "UObject/Object.h"
+#include "Math/Box.h"
+#include "Math/Range.h"
 #include "GeometryCacheTrack.generated.h"
 
+struct FResourceSizeEx;
+
 struct FGeometryCacheMeshData;
 struct FGeometryCacheTrackSampleInfo;
 
@@ -195,3 +196,7 @@ struct GEOMETRYCACHE_API FVisibilitySample
 	static const FVisibilitySample VisibleSample;
 	static const FVisibilitySample InvisibleSample;
 };
+
+#if UE_ENABLE_INCLUDE_ORDER_DEPRECATED_IN_5_2
+#include "CoreMinimal.h"
+#endif
diff --git a/Engine/Plugins/Runtime/GeometryCache/Source/GeometryCache/Classes/GeometryCacheTrackFlipbookAnimation.h b/Engine/Plugins/Runtime/GeometryCache/Source/GeometryCache/Classes/GeometryCacheTrackFlipbookAnimation.h
--- a/Engine/Plugins/Runtime/GeometryCache/Source/GeometryCache/Classes/GeometryCacheTrackFlipbookAnimation.h
+++ b/Engine/Plugins/Runtime/GeometryCache/Source/GeometryCache/Classes/GeometryCacheTrackFlipbookAnimation.h
@@ -1,11 +1,7 @@
 // Copyright Epic Games, Inc. All Rights Reserved.
 #pragma once
 
-#include "CoreMinimal.h"
-#include "UObject/ObjectMacros.h"
-#include "UObject/ScriptMacros.h"
 #include "GeometryCacheTrack.h"
-#include "GeometryCacheMeshData.h"
 
 #include "GeometryCacheTrackFlipbookAnimation.generated.h"
 
@@ -47,3 +43,8 @@ private:
 	TArray<float> MeshSampleTimes;
 };
 
+
+#if UE_ENABLE_INCLUDE_ORDER_DEPRECATED_IN_5_2
+#include "CoreMinimal.h"
+#include "GeometryCacheMeshData.h"
+#endif
diff --git a/Engine/Plugins/Runtime/GeometryCache/Source/GeometryCache/Classes/GeometryCacheTrackStreamable.h b/Engine/Plugins/Runtime/GeometryCache/Source/GeometryCache/Classes/GeometryCacheTrackStreamable.h
--- a/Engine/Plugins/Runtime/GeometryCache/Source/GeometryCache/Classes/GeometryCacheTrackStreamable.h
+++ b/Engine/Plugins/Runtime/GeometryCache/Source/GeometryCache/Classes/GeometryCacheTrackStreamable.h
@@ -1,17 +1,15 @@
 // Copyright Epic Games, Inc. All Rights Reserved.
 #pragma once
 
-#include "CoreMinimal.h"
-#include "UObject/ObjectMacros.h"
-#include "UObject/ScriptMacros.h"
 #include "GeometryCacheCodecBase.h"
 #include "GeometryCacheTrack.h"
-#include "GeometryCacheMeshData.h"
 #include "RenderResource.h"
 #include "RenderCommandFence.h"
 
 #include "GeometryCacheTrackStreamable.generated.h"
 
+class UGeometryCacheTrackStreamable;
+
 class FGeometryCachePreprocessor;
 
 struct FGeometryCacheTrackMeshDataUpdate
@@ -227,3 +225,8 @@ private:
 	uint64 Hash;
 };
 
+
+#if UE_ENABLE_INCLUDE_ORDER_DEPRECATED_IN_5_2
+#include "CoreMinimal.h"
+#include "GeometryCacheMeshData.h"
+#endif
diff --git a/Engine/Plugins/Runtime/GeometryCache/Source/GeometryCache/Classes/GeometryCacheTrackTransformAnimation.h b/Engine/Plugins/Runtime/GeometryCache/Source/GeometryCache/Classes/GeometryCacheTrackTransformAnimation.h
--- a/Engine/Plugins/Runtime/GeometryCache/Source/GeometryCache/Classes/GeometryCacheTrackTransformAnimation.h
+++ b/Engine/Plugins/Runtime/GeometryCache/Source/GeometryCache/Classes/GeometryCacheTrackTransformAnimation.h
@@ -1,9 +1,6 @@
 // Copyright Epic Games, Inc. All Rights Reserved.
 #pragma once
 
-#include "CoreMinimal.h"
-#include "UObject/ObjectMacros.h"
-#include "UObject/ScriptMacros.h"
 #include "GeometryCacheTrack.h"
 #include "GeometryCacheMeshData.h"
 
@@ -33,3 +30,7 @@ class GEOMETRYCACHE_API UDEPRECATED_GeometryCacheTrack_TransformAnimation : publ
 private:
 	FGeometryCacheMeshData MeshData;
 };
+
+#if UE_ENABLE_INCLUDE_ORDER_DEPRECATED_IN_5_2
+#include "CoreMinimal.h"
+#endif
diff --git a/Engine/Plugins/Runtime/GeometryCache/Source/GeometryCache/Classes/NiagaraGeometryCacheRendererProperties.h b/Engine/Plugins/Runtime/GeometryCache/Source/GeometryCache/Classes/NiagaraGeometryCacheRendererProperties.h
--- a/Engine/Plugins/Runtime/GeometryCache/Source/GeometryCache/Classes/NiagaraGeometryCacheRendererProperties.h
+++ b/Engine/Plugins/Runtime/GeometryCache/Source/GeometryCache/Classes/NiagaraGeometryCacheRendererProperties.h
@@ -2,12 +2,13 @@
 
 #pragma once
 
-#include "CoreMinimal.h"
-#include "GeometryCache.h"
-#include "NiagaraCommon.h"
 #include "NiagaraRendererProperties.h"
 #include "NiagaraGeometryCacheRendererProperties.generated.h"
 
+class UGeometryCache;
+struct FNiagaraVariable;
+struct FVersionedNiagaraEmitter;
+
 class FNiagaraEmitterInstance;
 class SWidget;
 
@@ -125,3 +126,8 @@ protected:
 private:
 	static TArray<TWeakObjectPtr<UNiagaraGeometryCacheRendererProperties>> GeometryCacheRendererPropertiesToDeferredInit;
 };
+
+#if UE_ENABLE_INCLUDE_ORDER_DEPRECATED_IN_5_2
+#include "CoreMinimal.h"
+#include "GeometryCache.h"
+#endif
diff --git a/Engine/Plugins/Runtime/GeometryCache/Source/GeometryCache/Classes/NiagaraRendererGeometryCache.h b/Engine/Plugins/Runtime/GeometryCache/Source/GeometryCache/Classes/NiagaraRendererGeometryCache.h
--- a/Engine/Plugins/Runtime/GeometryCache/Source/GeometryCache/Classes/NiagaraRendererGeometryCache.h
+++ b/Engine/Plugins/Runtime/GeometryCache/Source/GeometryCache/Classes/NiagaraRendererGeometryCache.h
@@ -6,8 +6,8 @@ NiagaraRendererGeometryCache.h: Renderer for geometry cache components.
 #pragma once
 
 #include "NiagaraRenderer.h"
-#include "CoreMinimal.h"
-#include "GeometryCacheComponent.h"
+
+class UGeometryCacheComponent;
 
 class UNiagaraGeometryCacheRendererProperties;
 class FNiagaraDataSet;
@@ -44,4 +44,9 @@ private:
 
 	// all of the spawned components
 	TArray<FComponentPoolEntry> ComponentPool;
-};
\ No newline at end of file
+};
+
+#if UE_ENABLE_INCLUDE_ORDER_DEPRECATED_IN_5_2
+#include "CoreMinimal.h"
+#include "GeometryCacheComponent.h"
+#endif
diff --git a/Engine/Plugins/Runtime/GeometryCache/Source/GeometryCache/Public/GeometryCacheHelpers.h b/Engine/Plugins/Runtime/GeometryCache/Source/GeometryCache/Public/GeometryCacheHelpers.h
--- a/Engine/Plugins/Runtime/GeometryCache/Source/GeometryCache/Public/GeometryCacheHelpers.h
+++ b/Engine/Plugins/Runtime/GeometryCache/Source/GeometryCache/Public/GeometryCacheHelpers.h
@@ -2,7 +2,7 @@
 
 #pragma once
 
-#include "Math/UnrealMath.h"
+#include "Math/UnrealMathUtility.h"
 
 struct GeometyCacheHelpers 
 {
@@ -13,4 +13,8 @@ struct GeometyCacheHelpers
 	{
 		return Time - Duration * FMath::FloorToFloat(Time / Duration);
 	}
-};
\ No newline at end of file
+};
+
+#if UE_ENABLE_INCLUDE_ORDER_DEPRECATED_IN_5_2
+#include "Math/UnrealMath.h"
+#endif
diff --git a/Engine/Plugins/Runtime/GeometryCache/Source/GeometryCache/Public/GeometryCacheModule.h b/Engine/Plugins/Runtime/GeometryCache/Source/GeometryCache/Public/GeometryCacheModule.h
--- a/Engine/Plugins/Runtime/GeometryCache/Source/GeometryCache/Public/GeometryCacheModule.h
+++ b/Engine/Plugins/Runtime/GeometryCache/Source/GeometryCache/Public/GeometryCacheModule.h
@@ -2,9 +2,6 @@
 
 #pragma once
 
-#include "CoreMinimal.h"
-#include "HAL/LowLevelMemTracker.h"
-#include "Modules/ModuleInterface.h"
 #include "Modules/ModuleManager.h"
 #include "Stats/Stats.h"
 
@@ -44,4 +41,8 @@ public:
 private:
 };
 
-DECLARE_STATS_GROUP(TEXT("GeometryCache"), STATGROUP_GeometryCache, STATCAT_Advanced);
\ No newline at end of file
+DECLARE_STATS_GROUP(TEXT("GeometryCache"), STATGROUP_GeometryCache, STATCAT_Advanced);
+
+#if UE_ENABLE_INCLUDE_ORDER_DEPRECATED_IN_5_2
+#include "CoreMinimal.h"
+#endif
diff --git a/Engine/Plugins/Runtime/GeometryCache/Source/GeometryCache/Public/GeometryCacheSceneProxy.h b/Engine/Plugins/Runtime/GeometryCache/Source/GeometryCache/Public/GeometryCacheSceneProxy.h
--- a/Engine/Plugins/Runtime/GeometryCache/Source/GeometryCache/Public/GeometryCacheSceneProxy.h
+++ b/Engine/Plugins/Runtime/GeometryCache/Source/GeometryCache/Public/GeometryCacheSceneProxy.h
@@ -1,19 +1,16 @@
 // Copyright Epic Games, Inc. All Rights Reserved.
 #pragma once
 
-#include "CoreMinimal.h"
-#include "HitProxies.h"
-#include "Stats/Stats.h"
-#include "Containers/ResourceArray.h"
-#include "RenderResource.h"
+#include "Materials/MaterialRelevance.h"
 #include "RayTracingGeometry.h"
-#include "PrimitiveViewRelevance.h"
 #include "PrimitiveSceneProxy.h"
 #include "Materials/MaterialInterface.h"
 #include "GeometryCacheVertexFactory.h"
 #include "DynamicMeshBuilder.h"
-#include "StaticMeshResources.h"
-#include "Logging/LogMacros.h"
+
+class FDynamicPrimitiveUniformBuffer;
+class FHitProxyId;
+struct FRayTracingMaterialGatheringContext;
 
 class FMeshElementCollector;
 struct FGeometryCacheMeshData;
@@ -410,3 +407,9 @@ private:
 	/** Function used to create a new track proxy at construction */
 	TFunction<FGeomCacheTrackProxy*()> CreateTrackProxy;
 };
+
+#if UE_ENABLE_INCLUDE_ORDER_DEPRECATED_IN_5_2
+#include "CoreMinimal.h"
+#include "PrimitiveViewRelevance.h"
+#include "StaticMeshResources.h"
+#endif
diff --git a/Engine/Plugins/Runtime/GeometryCache/Source/GeometryCacheEd/Classes/AssetTypeActions_GeometryCache.h b/Engine/Plugins/Runtime/GeometryCache/Source/GeometryCacheEd/Classes/AssetTypeActions_GeometryCache.h
--- a/Engine/Plugins/Runtime/GeometryCache/Source/GeometryCacheEd/Classes/AssetTypeActions_GeometryCache.h
+++ b/Engine/Plugins/Runtime/GeometryCache/Source/GeometryCacheEd/Classes/AssetTypeActions_GeometryCache.h
@@ -1,8 +1,6 @@
 // Copyright Epic Games, Inc. All Rights Reserved.
 #pragma once
 
-#include "CoreMinimal.h"
-#include "Toolkits/IToolkitHost.h"
 #include "AssetTypeActions_Base.h"
 
 
@@ -21,3 +19,8 @@ public:
 	virtual void GetResolvedSourceFilePaths(const TArray<UObject*>& TypeAssets, TArray<FString>& OutSourceFilePaths) const override;
 	// End IAssetTypeActions
 };
+
+#if UE_ENABLE_INCLUDE_ORDER_DEPRECATED_IN_5_2
+#include "CoreMinimal.h"
+#include "Toolkits/IToolkitHost.h"
+#endif
diff --git a/Engine/Plugins/Runtime/GeometryCache/Source/GeometryCacheEd/Classes/GeometryCacheThumbnailRenderer.h b/Engine/Plugins/Runtime/GeometryCache/Source/GeometryCacheEd/Classes/GeometryCacheThumbnailRenderer.h
--- a/Engine/Plugins/Runtime/GeometryCache/Source/GeometryCacheEd/Classes/GeometryCacheThumbnailRenderer.h
+++ b/Engine/Plugins/Runtime/GeometryCache/Source/GeometryCacheEd/Classes/GeometryCacheThumbnailRenderer.h
@@ -7,11 +7,11 @@
 
 #pragma once
 
-#include "CoreMinimal.h"
-#include "UObject/ObjectMacros.h"
 #include "ThumbnailRendering/DefaultSizedThumbnailRenderer.h"
 #include "GeometryCacheThumbnailRenderer.generated.h"
 
+enum class EThumbnailRenderFrequency : uint8;
+
 class FCanvas;
 class FGeometryCacheThumbnailScene;
 class FRenderTarget;
@@ -34,3 +34,7 @@ public:
 	virtual EThumbnailRenderFrequency GetThumbnailRenderFrequency(UObject* Object) const override;
 
 };
+
+#if UE_ENABLE_INCLUDE_ORDER_DEPRECATED_IN_5_2
+#include "CoreMinimal.h"
+#endif
diff --git a/Engine/Plugins/Runtime/GeometryCache/Source/GeometryCacheSequencer/Classes/GeometryCacheTrackEditor.h b/Engine/Plugins/Runtime/GeometryCache/Source/GeometryCacheSequencer/Classes/GeometryCacheTrackEditor.h
--- a/Engine/Plugins/Runtime/GeometryCache/Source/GeometryCacheSequencer/Classes/GeometryCacheTrackEditor.h
+++ b/Engine/Plugins/Runtime/GeometryCache/Source/GeometryCacheSequencer/Classes/GeometryCacheTrackEditor.h
@@ -2,16 +2,11 @@
 
 #pragma once
 
-#include "CoreMinimal.h"
-#include "Misc/Guid.h"
-#include "Templates/SubclassOf.h"
-#include "Widgets/SWidget.h"
-#include "ISequencer.h"
-#include "MovieSceneTrack.h"
 #include "ISequencerSection.h"
-#include "ISequencerTrackEditor.h"
 #include "MovieSceneTrackEditor.h"
 
+struct FBuildEditWidgetParams;
+
 struct FAssetData;
 class FMenuBuilder;
 class FSequencerSectionPainter;
@@ -105,3 +100,7 @@ private:
 	/** Cached start time valid only during resize */
 	FFrameNumber InitialStartTimeDuringResize;
 };
+
+#if UE_ENABLE_INCLUDE_ORDER_DEPRECATED_IN_5_2
+#include "CoreMinimal.h"
+#endif
diff --git a/Engine/Plugins/Runtime/GeometryCache/Source/GeometryCacheSequencer/Public/GeometryCacheSequencerModule.h b/Engine/Plugins/Runtime/GeometryCache/Source/GeometryCacheSequencer/Public/GeometryCacheSequencerModule.h
--- a/Engine/Plugins/Runtime/GeometryCache/Source/GeometryCacheSequencer/Public/GeometryCacheSequencerModule.h
+++ b/Engine/Plugins/Runtime/GeometryCache/Source/GeometryCacheSequencer/Public/GeometryCacheSequencerModule.h
@@ -2,12 +2,9 @@
 
 #pragma once
 
-#include "CoreMinimal.h"
 #include "GeometryCacheModule.h"
 #include "GeometryCacheTrackEditor.h"
 #include "ISequencerModule.h"
-#include "Modules/ModuleInterface.h"
-#include "Modules/ModuleManager.h"
 
 /**
  * The public interface to this module
@@ -39,3 +36,7 @@ public:
 	FDelegateHandle TrackEditorBindingHandle;
 };
 
+
+#if UE_ENABLE_INCLUDE_ORDER_DEPRECATED_IN_5_2
+#include "CoreMinimal.h"
+#endif
diff --git a/Engine/Plugins/Runtime/GeometryCache/Source/GeometryCacheStreamer/Public/GeometryCacheStreamBase.h b/Engine/Plugins/Runtime/GeometryCache/Source/GeometryCacheStreamer/Public/GeometryCacheStreamBase.h
--- a/Engine/Plugins/Runtime/GeometryCache/Source/GeometryCacheStreamer/Public/GeometryCacheStreamBase.h
+++ b/Engine/Plugins/Runtime/GeometryCache/Source/GeometryCacheStreamer/Public/GeometryCacheStreamBase.h
@@ -2,8 +2,9 @@
 
 #pragma once
 
+#include "Containers/Map.h"
 #include "IGeometryCacheStream.h"
-#include <atomic>
+#include "HAL/CriticalSection.h"
 
 struct FGeometryCacheStreamReadRequest;
 
diff --git a/Engine/Plugins/Runtime/GeometryCache/Source/GeometryCacheStreamer/Public/GeometryCacheStreamerSettings.h b/Engine/Plugins/Runtime/GeometryCache/Source/GeometryCacheStreamer/Public/GeometryCacheStreamerSettings.h
--- a/Engine/Plugins/Runtime/GeometryCache/Source/GeometryCacheStreamer/Public/GeometryCacheStreamerSettings.h
+++ b/Engine/Plugins/Runtime/GeometryCache/Source/GeometryCacheStreamer/Public/GeometryCacheStreamerSettings.h
@@ -3,7 +3,6 @@
 #pragma once
 
 #include "Engine/DeveloperSettings.h"
-#include "UObject/ObjectMacros.h"
 
 #include "GeometryCacheStreamerSettings.generated.h"
 
diff --git a/Engine/Plugins/Runtime/GeometryCache/Source/GeometryCacheStreamer/Public/IGeometryCacheStream.h b/Engine/Plugins/Runtime/GeometryCache/Source/GeometryCacheStreamer/Public/IGeometryCacheStream.h
--- a/Engine/Plugins/Runtime/GeometryCache/Source/GeometryCacheStreamer/Public/IGeometryCacheStream.h
+++ b/Engine/Plugins/Runtime/GeometryCache/Source/GeometryCacheStreamer/Public/IGeometryCacheStream.h
@@ -2,7 +2,7 @@
 
 #pragma once
 
-#include "CoreMinimal.h"
+#include "Containers/ContainersFwd.h" // IWYU pragma: keep
 
 struct FGeometryCacheMeshData;
 struct FGeometryCacheStreamStats;
@@ -47,3 +47,7 @@ struct FGeometryCacheStreamStats
 	float MemoryUsed = 0.0f;	// in MB
 	float AverageBitrate = 0.0f;// in MB/s
 };
+
+#if UE_ENABLE_INCLUDE_ORDER_DEPRECATED_IN_5_2
+#include "CoreMinimal.h"
+#endif
diff --git a/Engine/Plugins/Runtime/GeometryCache/Source/GeometryCacheStreamer/Public/IGeometryCacheStreamer.h b/Engine/Plugins/Runtime/GeometryCache/Source/GeometryCacheStreamer/Public/IGeometryCacheStreamer.h
--- a/Engine/Plugins/Runtime/GeometryCache/Source/GeometryCacheStreamer/Public/IGeometryCacheStreamer.h
+++ b/Engine/Plugins/Runtime/GeometryCache/Source/GeometryCacheStreamer/Public/IGeometryCacheStreamer.h
@@ -2,7 +2,7 @@
 
 #pragma once
 
-#include "CoreMinimal.h"
+#include "HAL/Platform.h"
 
 struct FGeometryCacheMeshData;
 class IGeometryCacheStream;
@@ -28,3 +28,7 @@ public:
 	 */
 	virtual bool TryGetFrameData(UGeometryCacheTrack* Track, int32 FrameIndex, FGeometryCacheMeshData& OutMeshData) = 0;
 };
+
+#if UE_ENABLE_INCLUDE_ORDER_DEPRECATED_IN_5_2
+#include "CoreMinimal.h"
+#endif
diff --git a/Engine/Plugins/Runtime/GeometryCache/Source/GeometryCacheTracks/Classes/MovieSceneGeometryCacheSection.h b/Engine/Plugins/Runtime/GeometryCache/Source/GeometryCacheTracks/Classes/MovieSceneGeometryCacheSection.h
--- a/Engine/Plugins/Runtime/GeometryCache/Source/GeometryCacheTracks/Classes/MovieSceneGeometryCacheSection.h
+++ b/Engine/Plugins/Runtime/GeometryCache/Source/GeometryCacheTracks/Classes/MovieSceneGeometryCacheSection.h
@@ -2,15 +2,12 @@
 
 #pragma once
 
-#include "CoreMinimal.h"
-#include "UObject/ObjectMacros.h"
 #include "MovieSceneSection.h"
-#include "GeometryCacheComponent.h"
-#include "GeometryCache.h"
-#include "Channels/MovieSceneFloatChannel.h"
-#include "UObject/SoftObjectPath.h"
 #include "MovieSceneGeometryCacheSection.generated.h"
 
+class UGeometryCache;
+struct FQualifiedFrameTime;
+
 USTRUCT(BlueprintType)
  struct GEOMETRYCACHETRACKS_API FMovieSceneGeometryCacheParams
 {
@@ -98,3 +95,10 @@ public:
 #endif
 
 };
+
+#if UE_ENABLE_INCLUDE_ORDER_DEPRECATED_IN_5_2
+#include "Channels/MovieSceneFloatChannel.h"
+#include "CoreMinimal.h"
+#include "GeometryCache.h"
+#include "GeometryCacheComponent.h"
+#endif
diff --git a/Engine/Plugins/Runtime/ModelViewViewModel/Source/ModelViewViewModelEditor/Private/Widgets/SMVVMFunctionParameter.cpp b/Engine/Plugins/Runtime/ModelViewViewModel/Source/ModelViewViewModelEditor/Private/Widgets/SMVVMFunctionParameter.cpp
--- a/Engine/Plugins/Runtime/ModelViewViewModel/Source/ModelViewViewModelEditor/Private/Widgets/SMVVMFunctionParameter.cpp
+++ b/Engine/Plugins/Runtime/ModelViewViewModel/Source/ModelViewViewModelEditor/Private/Widgets/SMVVMFunctionParameter.cpp
@@ -3,12 +3,11 @@
 #include "SMVVMFunctionParameter.h" 
 #include "EdGraphSchema_K2.h"
 #include "Editor.h"
-#include "MVVMBlueprintView.h"
+#include "MVVMBlueprintViewBinding.h"
 #include "MVVMEditorSubsystem.h"
 #include "NodeFactory.h"
 #include "SGraphPin.h"
 #include "Styling/MVVMEditorStyle.h"
-#include "Types/MVVMBindingMode.h"
 #include "WidgetBlueprint.h"
 #include "Widgets/Images/SImage.h"
 #include "Widgets/SBoxPanel.h"
diff --git a/Engine/Plugins/Runtime/ModelViewViewModel/Source/ModelViewViewModelEditor/Private/Widgets/SMVVMViewModelBindingListWidget.cpp b/Engine/Plugins/Runtime/ModelViewViewModel/Source/ModelViewViewModelEditor/Private/Widgets/SMVVMViewModelBindingListWidget.cpp
--- a/Engine/Plugins/Runtime/ModelViewViewModel/Source/ModelViewViewModelEditor/Private/Widgets/SMVVMViewModelBindingListWidget.cpp
+++ b/Engine/Plugins/Runtime/ModelViewViewModel/Source/ModelViewViewModelEditor/Private/Widgets/SMVVMViewModelBindingListWidget.cpp
@@ -11,7 +11,7 @@
 #include "Types/MVVMAvailableBinding.h"
 #include "Types/MVVMBindingSource.h"
 #include "WidgetBlueprint.h"
-#include "Widgets/Images/SLayeredImage.h"
+#include "Widgets/Images/SImage.h"
 #include "Widgets/Layout/SBox.h"
 
 #define LOCTEXT_NAMESPACE "SSourceBindingList"
diff --git a/Engine/Plugins/Runtime/NetworkPrediction/Source/NetworkPrediction/Public/NetworkPredictionProxyInit.h b/Engine/Plugins/Runtime/NetworkPrediction/Source/NetworkPrediction/Public/NetworkPredictionProxyInit.h
--- a/Engine/Plugins/Runtime/NetworkPrediction/Source/NetworkPrediction/Public/NetworkPredictionProxyInit.h
+++ b/Engine/Plugins/Runtime/NetworkPrediction/Source/NetworkPrediction/Public/NetworkPredictionProxyInit.h
@@ -5,7 +5,6 @@
 #include "Engine/World.h"
 #include "NetworkPredictionProxy.h"
 #include "NetworkPredictionWorldManager.h"
-#include "Trace/Trace.inl"
 
 struct FNetworkPredictionSettings;
 
diff --git a/Engine/Plugins/Runtime/RenderTrace/Source/Public/RenderTrace.h b/Engine/Plugins/Runtime/RenderTrace/Source/Public/RenderTrace.h
--- a/Engine/Plugins/Runtime/RenderTrace/Source/Public/RenderTrace.h
+++ b/Engine/Plugins/Runtime/RenderTrace/Source/Public/RenderTrace.h
@@ -5,6 +5,10 @@
 #include "Stats/Stats2.h"
 #include "Tickable.h"
 
+enum class EStatFlags : uint8;
+struct FStatGroup_STATGROUP_Tickables;
+struct TStatIdData;
+
 DECLARE_LOG_CATEGORY_EXTERN(LogRenderTrace, Log, All);
 
 DECLARE_DELEGATE_ThreeParams(FRenderTraceDelegate, uint32 TaskID, class UPhysicalMaterial const*, int64 UserData);
diff --git a/Engine/Plugins/Runtime/StateTree/Source/StateTreeEditorModule/Private/Commandlets/StateTreeCompileAllCommandlet.cpp b/Engine/Plugins/Runtime/StateTree/Source/StateTreeEditorModule/Private/Commandlets/StateTreeCompileAllCommandlet.cpp
--- a/Engine/Plugins/Runtime/StateTree/Source/StateTreeEditorModule/Private/Commandlets/StateTreeCompileAllCommandlet.cpp
+++ b/Engine/Plugins/Runtime/StateTree/Source/StateTreeEditorModule/Private/Commandlets/StateTreeCompileAllCommandlet.cpp
@@ -7,6 +7,7 @@
 #include "AssetRegistry/AssetRegistryModule.h"
 #include "SourceControlOperations.h"
 #include "SourceControlHelpers.h"
+#include "StateTreeCompilerLog.h"
 #include "StateTreeEditor.h"
 #include "StateTreeCompiler.h"
 #include "StateTreeDelegates.h"
diff --git a/Engine/Plugins/Runtime/StateTree/Source/StateTreeEditorModule/Private/Customizations/StateTreeTransitionDetails.cpp b/Engine/Plugins/Runtime/StateTree/Source/StateTreeEditorModule/Private/Customizations/StateTreeTransitionDetails.cpp
--- a/Engine/Plugins/Runtime/StateTree/Source/StateTreeEditorModule/Private/Customizations/StateTreeTransitionDetails.cpp
+++ b/Engine/Plugins/Runtime/StateTree/Source/StateTreeEditorModule/Private/Customizations/StateTreeTransitionDetails.cpp
@@ -1,13 +1,13 @@
 // Copyright Epic Games, Inc. All Rights Reserved.
 
 #include "StateTreeTransitionDetails.h"
+#include "StateTreeTypes.h"
 #include "Widgets/Layout/SBox.h"
 #include "Widgets/SBoxPanel.h"
 #include "Widgets/Text/STextBlock.h"
 #include "DetailWidgetRow.h"
 #include "DetailLayoutBuilder.h"
 #include "IDetailChildrenBuilder.h"
-#include "StateTreeState.h"
 #include "Widgets/Text/STextBlock.h"
 #include "Widgets/Layout/SBox.h"
 #include "StateTreePropertyHelpers.h"
diff --git a/Engine/Plugins/Runtime/StateTree/Source/StateTreeEditorModule/Private/StateTreeCompiler.cpp b/Engine/Plugins/Runtime/StateTree/Source/StateTreeEditorModule/Private/StateTreeCompiler.cpp
--- a/Engine/Plugins/Runtime/StateTree/Source/StateTreeEditorModule/Private/StateTreeCompiler.cpp
+++ b/Engine/Plugins/Runtime/StateTree/Source/StateTreeEditorModule/Private/StateTreeCompiler.cpp
@@ -3,6 +3,7 @@
 #include "StateTreeCompiler.h"
 #include "StateTree.h"
 #include "StateTreeAnyEnum.h"
+#include "StateTreeCompilerLog.h"
 #include "StateTreeEditorData.h"
 #include "StateTreeEvaluatorBase.h"
 #include "StateTreeTaskBase.h"
diff --git a/Engine/Plugins/Runtime/StateTree/Source/StateTreeEditorModule/Public/StateTreeCompiler.h b/Engine/Plugins/Runtime/StateTree/Source/StateTreeEditorModule/Public/StateTreeCompiler.h
--- a/Engine/Plugins/Runtime/StateTree/Source/StateTreeEditorModule/Public/StateTreeCompiler.h
+++ b/Engine/Plugins/Runtime/StateTree/Source/StateTreeEditorModule/Public/StateTreeCompiler.h
@@ -2,10 +2,14 @@
 
 #pragma once
 
-#include "StateTreeTypes.h"
-#include "InstancedStruct.h"
 #include "StateTreePropertyBindingCompiler.h"
-#include "StateTreeCompilerLog.h"
+
+struct FStructView;
+
+enum class EStateTreeConditionOperand : uint8;
+enum class EStateTreePropertyUsage : uint8;
+struct FStateTreeDataView;
+struct FStateTreeStateHandle;
 
 class UStateTree;
 class UStateTreeState;
@@ -152,4 +156,8 @@ namespace UE::StateTree::Compiler
 	EStateTreePropertyUsage GetUsageFromMetaData(const FProperty* Property);
 
 
-}; // UE::StateTree::Compiler
\ No newline at end of file
+}; // UE::StateTree::Compiler
+
+#if UE_ENABLE_INCLUDE_ORDER_DEPRECATED_IN_5_2
+#include "StateTreeCompilerLog.h"
+#endif
diff --git a/Engine/Plugins/Runtime/StateTree/Source/StateTreeEditorModule/Public/StateTreeCompilerLog.h b/Engine/Plugins/Runtime/StateTree/Source/StateTreeEditorModule/Public/StateTreeCompilerLog.h
--- a/Engine/Plugins/Runtime/StateTree/Source/StateTreeEditorModule/Public/StateTreeCompilerLog.h
+++ b/Engine/Plugins/Runtime/StateTree/Source/StateTreeEditorModule/Public/StateTreeCompilerLog.h
@@ -1,8 +1,7 @@
-﻿// Copyright Epic Games, Inc. All Rights Reserved.
+// Copyright Epic Games, Inc. All Rights Reserved.
 
 #pragma once
 
-#include "CoreMinimal.h"
 #include "Logging/TokenizedMessage.h"
 #include "StateTreePropertyBindings.h"
 #include "StateTreeCompilerLog.generated.h"
@@ -121,3 +120,7 @@ struct FStateTreeCompilerLogStateScope
 	FStateTreeCompilerLog& Log; 
 	const UStateTreeState* State = nullptr;
 };
+
+#if UE_ENABLE_INCLUDE_ORDER_DEPRECATED_IN_5_2
+#include "CoreMinimal.h"
+#endif
diff --git a/Engine/Plugins/Runtime/StateTree/Source/StateTreeEditorModule/Public/StateTreeEditorData.h b/Engine/Plugins/Runtime/StateTree/Source/StateTreeEditorModule/Public/StateTreeEditorData.h
--- a/Engine/Plugins/Runtime/StateTree/Source/StateTreeEditorModule/Public/StateTreeEditorData.h
+++ b/Engine/Plugins/Runtime/StateTree/Source/StateTreeEditorModule/Public/StateTreeEditorData.h
@@ -6,6 +6,9 @@
 #include "StateTreeEditorPropertyBindings.h"
 #include "StateTreeEditorData.generated.h"
 
+struct FStateTreeBindableStructDesc;
+struct FStateTreeEditorPropertyPath;
+
 class UStateTreeSchema;
 
 UENUM()
diff --git a/Engine/Plugins/Runtime/StateTree/Source/StateTreeEditorModule/Public/StateTreeEditorModule.h b/Engine/Plugins/Runtime/StateTree/Source/StateTreeEditorModule/Public/StateTreeEditorModule.h
--- a/Engine/Plugins/Runtime/StateTree/Source/StateTreeEditorModule/Public/StateTreeEditorModule.h
+++ b/Engine/Plugins/Runtime/StateTree/Source/StateTreeEditorModule/Public/StateTreeEditorModule.h
@@ -2,8 +2,7 @@
 
 #pragma  once
 
-#include "CoreMinimal.h"
-#include "Toolkits/IToolkitHost.h"
+#include "Modules/ModuleInterface.h"
 #include "Toolkits/AssetEditorToolkit.h"
 
 class UStateTree;
@@ -39,3 +38,8 @@ protected:
 
 	TSharedPtr<FStateTreeNodeClassCache> NodeClassCache;
 };
+
+#if UE_ENABLE_INCLUDE_ORDER_DEPRECATED_IN_5_2
+#include "CoreMinimal.h"
+#include "Toolkits/IToolkitHost.h"
+#endif
diff --git a/Engine/Plugins/Runtime/StateTree/Source/StateTreeEditorModule/Public/StateTreeEditorNode.h b/Engine/Plugins/Runtime/StateTree/Source/StateTreeEditorModule/Public/StateTreeEditorNode.h
--- a/Engine/Plugins/Runtime/StateTree/Source/StateTreeEditorModule/Public/StateTreeEditorNode.h
+++ b/Engine/Plugins/Runtime/StateTree/Source/StateTreeEditorModule/Public/StateTreeEditorNode.h
@@ -1,11 +1,8 @@
-﻿// Copyright Epic Games, Inc. All Rights Reserved.
+// Copyright Epic Games, Inc. All Rights Reserved.
 
 #pragma once
 
-#include "StateTreeTypes.h"
 #include "StateTreeNodeBase.h"
-#include "Misc/Guid.h"
-#include "InstancedStruct.h"
 #include "StateTreeEditorNode.generated.h"
 
 UENUM()
diff --git a/Engine/Plugins/Runtime/StateTree/Source/StateTreeEditorModule/Public/StateTreeEditorPropertyBindings.h b/Engine/Plugins/Runtime/StateTree/Source/StateTreeEditorModule/Public/StateTreeEditorPropertyBindings.h
--- a/Engine/Plugins/Runtime/StateTree/Source/StateTreeEditorModule/Public/StateTreeEditorPropertyBindings.h
+++ b/Engine/Plugins/Runtime/StateTree/Source/StateTreeEditorModule/Public/StateTreeEditorPropertyBindings.h
@@ -2,9 +2,7 @@
 
 #pragma once
 
-#include "CoreMinimal.h"
 #include "UObject/Interface.h"
-#include "Misc/Guid.h"
 #include "StateTreePropertyBindings.h"
 #include "StateTreeEditorPropertyBindings.generated.h"
 
@@ -137,3 +135,7 @@ struct STATETREEEDITORMODULE_API FStateTreeBindingLookup : public IStateTreeBind
 
 	IStateTreeEditorPropertyBindingsOwner* BindingOwner = nullptr;
 };
+
+#if UE_ENABLE_INCLUDE_ORDER_DEPRECATED_IN_5_2
+#include "CoreMinimal.h"
+#endif
diff --git a/Engine/Plugins/Runtime/StateTree/Source/StateTreeEditorModule/Public/StateTreeEditorStyle.h b/Engine/Plugins/Runtime/StateTree/Source/StateTreeEditorModule/Public/StateTreeEditorStyle.h
--- a/Engine/Plugins/Runtime/StateTree/Source/StateTreeEditorModule/Public/StateTreeEditorStyle.h
+++ b/Engine/Plugins/Runtime/StateTree/Source/StateTreeEditorModule/Public/StateTreeEditorStyle.h
@@ -4,6 +4,8 @@
 
 #include "Styling/SlateStyle.h"
 
+class ISlateStyle;
+
 class STATETREEEDITORMODULE_API FStateTreeEditorStyle
 {
 public:
diff --git a/Engine/Plugins/Runtime/StateTree/Source/StateTreeEditorModule/Public/StateTreeNodeClassCache.h b/Engine/Plugins/Runtime/StateTree/Source/StateTreeEditorModule/Public/StateTreeNodeClassCache.h
--- a/Engine/Plugins/Runtime/StateTree/Source/StateTreeEditorModule/Public/StateTreeNodeClassCache.h
+++ b/Engine/Plugins/Runtime/StateTree/Source/StateTreeEditorModule/Public/StateTreeNodeClassCache.h
@@ -2,11 +2,10 @@
 
 #pragma once
 
-#include "CoreMinimal.h"
-#include "UObject/ObjectMacros.h"
-#include "UObject/WeakObjectPtr.h"
 #include "StateTreeNodeClassCache.generated.h"
 
+enum class EReloadCompleteReason;
+
 struct FAssetData;
 
 /**
@@ -116,3 +115,7 @@ private:
 	TArray<FRootClassContainer> RootClasses;
 	TMap<FString, int32> RootClassNameToIndex;
 };
+
+#if UE_ENABLE_INCLUDE_ORDER_DEPRECATED_IN_5_2
+#include "CoreMinimal.h"
+#endif
diff --git a/Engine/Plugins/Runtime/StateTree/Source/StateTreeEditorModule/Public/StateTreePropertyBindingCompiler.h b/Engine/Plugins/Runtime/StateTree/Source/StateTreeEditorModule/Public/StateTreePropertyBindingCompiler.h
--- a/Engine/Plugins/Runtime/StateTree/Source/StateTreeEditorModule/Public/StateTreePropertyBindingCompiler.h
+++ b/Engine/Plugins/Runtime/StateTree/Source/StateTreeEditorModule/Public/StateTreePropertyBindingCompiler.h
@@ -2,16 +2,13 @@
 
 #pragma once
 
-#include "CoreMinimal.h"
-#include "Misc/Guid.h"
-#include "InstancedStruct.h"
-#include "UObject/Interface.h"
-#include "IPropertyAccessEditor.h"
-#include "StateTreeEditorPropertyBindings.h"
 #include "StateTreePropertyBindings.h"
-#include "StateTreeCompilerLog.h"
 #include "StateTreePropertyBindingCompiler.generated.h"
 
+enum class EPropertyAccessCompatibility;
+struct FStateTreeCompilerLog;
+struct FStateTreeEditorPropertyBinding;
+
 /**
  * Helper class to compile editor representation of property bindings into runtime representation.
  * TODO: Better error reporting, something that can be shown in the UI.
@@ -94,3 +91,11 @@ protected:
 
 	FStateTreeCompilerLog* Log = nullptr;
 };
+
+#if UE_ENABLE_INCLUDE_ORDER_DEPRECATED_IN_5_2
+#include "CoreMinimal.h"
+#include "IPropertyAccessEditor.h"
+#include "StateTreeCompilerLog.h"
+#include "StateTreeEditorPropertyBindings.h"
+#include "UObject/Interface.h"
+#endif
diff --git a/Engine/Plugins/Runtime/StateTree/Source/StateTreeEditorModule/Public/StateTreeState.h b/Engine/Plugins/Runtime/StateTree/Source/StateTreeEditorModule/Public/StateTreeState.h
--- a/Engine/Plugins/Runtime/StateTree/Source/StateTreeEditorModule/Public/StateTreeState.h
+++ b/Engine/Plugins/Runtime/StateTree/Source/StateTreeEditorModule/Public/StateTreeState.h
@@ -2,11 +2,6 @@
 
 #pragma once
 
-#include "StateTreeTypes.h"
-#include "StateTreeNodeBase.h"
-#include "Misc/Guid.h"
-#include "InstancedStruct.h"
-#include "PropertyBag.h"
 #include "StateTreeEditorNode.h"
 #include "StateTreeState.generated.h"
 
diff --git a/Engine/Plugins/Runtime/StateTree/Source/StateTreeEditorModule/Public/StateTreeToolMenuContext.h b/Engine/Plugins/Runtime/StateTree/Source/StateTreeEditorModule/Public/StateTreeToolMenuContext.h
--- a/Engine/Plugins/Runtime/StateTree/Source/StateTreeEditorModule/Public/StateTreeToolMenuContext.h
+++ b/Engine/Plugins/Runtime/StateTree/Source/StateTreeEditorModule/Public/StateTreeToolMenuContext.h
@@ -1,9 +1,7 @@
-﻿// Copyright Epic Games, Inc. All Rights Reserved.
+// Copyright Epic Games, Inc. All Rights Reserved.
 
 #pragma once
 
-#include "UObject/Object.h"
-#include "Templates/SharedPointer.h"
 #include "StateTreeToolMenuContext.generated.h"
 
 class FStateTreeEditor;
diff --git a/Engine/Plugins/Runtime/StateTree/Source/StateTreeModule/Private/StateTreeInstanceData.cpp b/Engine/Plugins/Runtime/StateTree/Source/StateTreeModule/Private/StateTreeInstanceData.cpp
--- a/Engine/Plugins/Runtime/StateTree/Source/StateTreeModule/Private/StateTreeInstanceData.cpp
+++ b/Engine/Plugins/Runtime/StateTree/Source/StateTreeModule/Private/StateTreeInstanceData.cpp
@@ -1,7 +1,6 @@
 // Copyright Epic Games, Inc. All Rights Reserved.
 
 #include "StateTreeInstanceData.h"
-#include "StateTreeTypes.h"
 #include "Serialization/MemoryReader.h"
 #include "Serialization/MemoryWriter.h"
 #include "VisualLogger/VisualLogger.h"
diff --git a/Engine/Plugins/Runtime/StateTree/Source/StateTreeModule/Public/Blueprint/StateTreeEvaluatorBlueprintBase.h b/Engine/Plugins/Runtime/StateTree/Source/StateTreeModule/Public/Blueprint/StateTreeEvaluatorBlueprintBase.h
--- a/Engine/Plugins/Runtime/StateTree/Source/StateTreeModule/Public/Blueprint/StateTreeEvaluatorBlueprintBase.h
+++ b/Engine/Plugins/Runtime/StateTree/Source/StateTreeModule/Public/Blueprint/StateTreeEvaluatorBlueprintBase.h
@@ -2,11 +2,8 @@
 
 #pragma once
 
-#include "CoreMinimal.h"
 #include "Templates/SubclassOf.h"
-#include "StateTreeTypes.h"
 #include "StateTreeEvaluatorBase.h"
-#include "StateTreeEvents.h"
 #include "StateTreeNodeBlueprintBase.h"
 #include "StateTreeEvaluatorBlueprintBase.generated.h"
 
@@ -60,3 +57,8 @@ struct STATETREEMODULE_API FStateTreeBlueprintEvaluatorWrapper : public FStateTr
 	UPROPERTY()
 	TSubclassOf<UStateTreeEvaluatorBlueprintBase> EvaluatorClass = nullptr;
 };
+
+#if UE_ENABLE_INCLUDE_ORDER_DEPRECATED_IN_5_2
+#include "CoreMinimal.h"
+#include "StateTreeEvents.h"
+#endif
diff --git a/Engine/Plugins/Runtime/StateTree/Source/StateTreeModule/Public/Blueprint/StateTreeNodeBlueprintBase.h b/Engine/Plugins/Runtime/StateTree/Source/StateTreeModule/Public/Blueprint/StateTreeNodeBlueprintBase.h
--- a/Engine/Plugins/Runtime/StateTree/Source/StateTreeModule/Public/Blueprint/StateTreeNodeBlueprintBase.h
+++ b/Engine/Plugins/Runtime/StateTree/Source/StateTreeModule/Public/Blueprint/StateTreeNodeBlueprintBase.h
@@ -3,10 +3,11 @@
 #pragma once
 
 #include "StateTreeTypes.h"
-#include "StateTreeEvents.h"
-#include "StateTreeExecutionContext.h"
 #include "StateTreeNodeBlueprintBase.generated.h"
 
+struct FStateTreeEvent;
+struct FStateTreeEventQueue;
+struct FStateTreeInstanceStorage;
 struct FStateTreeLinker;
 struct FStateTreeExecutionContext;
 
@@ -68,3 +69,8 @@ private:
 	UPROPERTY()
 	mutable TObjectPtr<UObject> CachedOwner = nullptr;
 };
+
+#if UE_ENABLE_INCLUDE_ORDER_DEPRECATED_IN_5_2
+#include "StateTreeEvents.h"
+#include "StateTreeExecutionContext.h"
+#endif
diff --git a/Engine/Plugins/Runtime/StateTree/Source/StateTreeModule/Public/Blueprint/StateTreeTaskBlueprintBase.h b/Engine/Plugins/Runtime/StateTree/Source/StateTreeModule/Public/Blueprint/StateTreeTaskBlueprintBase.h
--- a/Engine/Plugins/Runtime/StateTree/Source/StateTreeModule/Public/Blueprint/StateTreeTaskBlueprintBase.h
+++ b/Engine/Plugins/Runtime/StateTree/Source/StateTreeModule/Public/Blueprint/StateTreeTaskBlueprintBase.h
@@ -1,12 +1,9 @@
-﻿// Copyright Epic Games, Inc. All Rights Reserved.
+// Copyright Epic Games, Inc. All Rights Reserved.
 
 #pragma once
 
-#include "CoreMinimal.h"
 #include "Templates/SubclassOf.h"
-#include "StateTreeTypes.h"
 #include "StateTreeTaskBase.h"
-#include "StateTreeEvents.h"
 #include "StateTreeNodeBlueprintBase.h"
 #include "StateTreeTaskBlueprintBase.generated.h"
 
@@ -101,3 +98,8 @@ struct STATETREEMODULE_API FStateTreeBlueprintTaskWrapper : public FStateTreeTas
 	UPROPERTY()
 	TSubclassOf<UStateTreeTaskBlueprintBase> TaskClass = nullptr;
 };
+
+#if UE_ENABLE_INCLUDE_ORDER_DEPRECATED_IN_5_2
+#include "CoreMinimal.h"
+#include "StateTreeEvents.h"
+#endif
diff --git a/Engine/Plugins/Runtime/StateTree/Source/StateTreeModule/Public/Conditions/StateTreeCommonConditions.h b/Engine/Plugins/Runtime/StateTree/Source/StateTreeModule/Public/Conditions/StateTreeCommonConditions.h
--- a/Engine/Plugins/Runtime/StateTree/Source/StateTreeModule/Public/Conditions/StateTreeCommonConditions.h
+++ b/Engine/Plugins/Runtime/StateTree/Source/StateTreeModule/Public/Conditions/StateTreeCommonConditions.h
@@ -3,11 +3,12 @@
 #pragma once
 
 #include "AITypes.h"
-#include "InstancedStruct.h"
 #include "StateTreeConditionBase.h"
 #include "StateTreeAnyEnum.h"
 #include "StateTreeCommonConditions.generated.h"
 
+struct FStateTreeDataView;
+
 USTRUCT()
 struct STATETREEMODULE_API FStateTreeCompareIntConditionInstanceData
 {
diff --git a/Engine/Plugins/Runtime/StateTree/Source/StateTreeModule/Public/Conditions/StateTreeGameplayTagConditions.h b/Engine/Plugins/Runtime/StateTree/Source/StateTreeModule/Public/Conditions/StateTreeGameplayTagConditions.h
--- a/Engine/Plugins/Runtime/StateTree/Source/StateTreeModule/Public/Conditions/StateTreeGameplayTagConditions.h
+++ b/Engine/Plugins/Runtime/StateTree/Source/StateTreeModule/Public/Conditions/StateTreeGameplayTagConditions.h
@@ -1,9 +1,8 @@
-﻿// Copyright Epic Games, Inc. All Rights Reserved.
+// Copyright Epic Games, Inc. All Rights Reserved.
 
 #pragma once
 
 #include "StateTreeConditionBase.h"
-#include "GameplayTagContainer.h"
 #include "StateTreeGameplayTagConditions.generated.h"
 
 USTRUCT()
@@ -108,4 +107,4 @@ struct STATETREEMODULE_API FGameplayTagQueryCondition : public FStateTreeConditi
 
 	UPROPERTY(EditAnywhere, Category = Condition)
 	bool bInvert = false;
-};
\ No newline at end of file
+};
diff --git a/Engine/Plugins/Runtime/StateTree/Source/StateTreeModule/Public/Conditions/StateTreeObjectConditions.h b/Engine/Plugins/Runtime/StateTree/Source/StateTreeModule/Public/Conditions/StateTreeObjectConditions.h
--- a/Engine/Plugins/Runtime/StateTree/Source/StateTreeModule/Public/Conditions/StateTreeObjectConditions.h
+++ b/Engine/Plugins/Runtime/StateTree/Source/StateTreeModule/Public/Conditions/StateTreeObjectConditions.h
@@ -1,8 +1,7 @@
-﻿// Copyright Epic Games, Inc. All Rights Reserved.
+// Copyright Epic Games, Inc. All Rights Reserved.
 
 #pragma once
 
-#include "UObject/Object.h"
 #include "StateTreeConditionBase.h"
 #include "StateTreeObjectConditions.generated.h"
 
diff --git a/Engine/Plugins/Runtime/StateTree/Source/StateTreeModule/Public/StateTree.h b/Engine/Plugins/Runtime/StateTree/Source/StateTreeModule/Public/StateTree.h
--- a/Engine/Plugins/Runtime/StateTree/Source/StateTreeModule/Public/StateTree.h
+++ b/Engine/Plugins/Runtime/StateTree/Source/StateTreeModule/Public/StateTree.h
@@ -3,13 +3,9 @@
 #pragma once
 
 #include "Engine/DataAsset.h"
-#include "StateTreeTypes.h"
 #include "StateTreeSchema.h"
-#include "InstancedStruct.h"
-#include "InstancedStructContainer.h"
 #include "StateTreePropertyBindings.h"
 #include "StateTreeInstanceData.h"
-#include "Misc/ScopeRWLock.h"
 #include "StateTree.generated.h"
 
 
@@ -244,3 +240,7 @@ private:
 #endif
 };
 
+
+#if UE_ENABLE_INCLUDE_ORDER_DEPRECATED_IN_5_2
+#include "Misc/ScopeRWLock.h"
+#endif
diff --git a/Engine/Plugins/Runtime/StateTree/Source/StateTreeModule/Public/StateTreeConditionBase.h b/Engine/Plugins/Runtime/StateTree/Source/StateTreeModule/Public/StateTreeConditionBase.h
--- a/Engine/Plugins/Runtime/StateTree/Source/StateTreeModule/Public/StateTreeConditionBase.h
+++ b/Engine/Plugins/Runtime/StateTree/Source/StateTreeModule/Public/StateTreeConditionBase.h
@@ -2,18 +2,12 @@
 
 #pragma once
 
-#include "StateTreeTypes.h"
 #include "StateTreeNodeBase.h"
-#include "StateTreeExecutionContext.h"
-#if WITH_EDITOR
-#include "StateTreePropertyBindings.h"
-#endif
 #include "StateTreeConditionBase.generated.h"
 
-#if WITH_EDITOR
-struct IStateTreeBindingLookup;
 struct FStateTreeEditorPropertyPath;
-#endif
+struct FStateTreeExecutionContext;
+struct IStateTreeBindingLookup;
 
 enum class EStateTreeCompare : uint8
 {
@@ -73,3 +67,8 @@ struct TStructOpsTypeTraits<Type> : public TStructOpsTypeTraitsBase2<Type> \
 		WithNoDestructor = true, \
 	}; \
 };
+
+#if UE_ENABLE_INCLUDE_ORDER_DEPRECATED_IN_5_2
+#include "StateTreeExecutionContext.h"
+#include "StateTreePropertyBindings.h"
+#endif
diff --git a/Engine/Plugins/Runtime/StateTree/Source/StateTreeModule/Public/StateTreeEvaluatorBase.h b/Engine/Plugins/Runtime/StateTree/Source/StateTreeModule/Public/StateTreeEvaluatorBase.h
--- a/Engine/Plugins/Runtime/StateTree/Source/StateTreeModule/Public/StateTreeEvaluatorBase.h
+++ b/Engine/Plugins/Runtime/StateTree/Source/StateTreeModule/Public/StateTreeEvaluatorBase.h
@@ -2,7 +2,6 @@
 
 #pragma once
 
-#include "StateTreeTypes.h"
 #include "StateTreeNodeBase.h"
 #include "StateTreeEvaluatorBase.generated.h"
 
@@ -49,4 +48,4 @@ USTRUCT(Meta=(Hidden))
 struct STATETREEMODULE_API FStateTreeEvaluatorCommonBase : public FStateTreeEvaluatorBase
 {
 	GENERATED_BODY()
-};
\ No newline at end of file
+};
diff --git a/Engine/Plugins/Runtime/StateTree/Source/StateTreeModule/Public/StateTreeEvents.h b/Engine/Plugins/Runtime/StateTree/Source/StateTreeModule/Public/StateTreeEvents.h
--- a/Engine/Plugins/Runtime/StateTree/Source/StateTreeModule/Public/StateTreeEvents.h
+++ b/Engine/Plugins/Runtime/StateTree/Source/StateTreeModule/Public/StateTreeEvents.h
@@ -1,9 +1,8 @@
-﻿// Copyright Epic Games, Inc. All Rights Reserved.
+// Copyright Epic Games, Inc. All Rights Reserved.
 
 #pragma once
 
 #include "GameplayTagContainer.h"
-#include "InstancedStruct.h"
 #include "StructView.h"
 #include "StateTreeEvents.generated.h"
 
diff --git a/Engine/Plugins/Runtime/StateTree/Source/StateTreeModule/Public/StateTreeExecutionContext.h b/Engine/Plugins/Runtime/StateTree/Source/StateTreeModule/Public/StateTreeExecutionContext.h
--- a/Engine/Plugins/Runtime/StateTree/Source/StateTreeModule/Public/StateTreeExecutionContext.h
+++ b/Engine/Plugins/Runtime/StateTree/Source/StateTreeModule/Public/StateTreeExecutionContext.h
@@ -3,10 +3,12 @@
 #pragma once
 
 #include "StateTree.h"
-#include "StateTreeInstanceData.h"
 #include "StateTreeNodeBase.h"
 #include "Experimental/ConcurrentLinearAllocator.h"
 
+struct FGameplayTag;
+struct FInstancedPropertyBag;
+
 struct FStateTreeEvaluatorBase;
 struct FStateTreeTaskBase;
 struct FStateTreeConditionBase;
diff --git a/Engine/Plugins/Runtime/StateTree/Source/StateTreeModule/Public/StateTreeInstanceData.h b/Engine/Plugins/Runtime/StateTree/Source/StateTreeModule/Public/StateTreeInstanceData.h
--- a/Engine/Plugins/Runtime/StateTree/Source/StateTreeModule/Public/StateTreeInstanceData.h
+++ b/Engine/Plugins/Runtime/StateTree/Source/StateTreeModule/Public/StateTreeInstanceData.h
@@ -3,9 +3,8 @@
 #pragma once
 
 #include "InstancedStructContainer.h"
-#include "InstancedStruct.h"
-#include "StructView.h"
 #include "StateTreeEvents.h"
+#include "StateTreeTypes.h"
 #include "StateTreeInstanceData.generated.h"
 
 struct FStateTreeTransitionRequest;
@@ -269,4 +268,4 @@ struct TStateTreeInstanceDataStructRef
 protected:
 	FStateTreeInstanceStorage& Storage;
 	int32 StructIndex = INDEX_NONE;
-};
\ No newline at end of file
+};
diff --git a/Engine/Plugins/Runtime/StateTree/Source/StateTreeModule/Public/StateTreePropertyBindings.h b/Engine/Plugins/Runtime/StateTree/Source/StateTreeModule/Public/StateTreePropertyBindings.h
--- a/Engine/Plugins/Runtime/StateTree/Source/StateTreeModule/Public/StateTreePropertyBindings.h
+++ b/Engine/Plugins/Runtime/StateTree/Source/StateTreeModule/Public/StateTreePropertyBindings.h
@@ -2,10 +2,7 @@
 
 #pragma once
 
-#include "CoreMinimal.h"
 #include "StateTreeTypes.h"
-#include "Misc/Guid.h"
-#include "StructView.h"
 #include "StateTreePropertyBindings.generated.h"
 
 class FProperty;
@@ -418,3 +415,7 @@ struct STATETREEMODULE_API IStateTreeBindingLookup
 	/** @return Leaf property based on property path. */
 	virtual const FProperty* GetPropertyPathLeafProperty(const FStateTreeEditorPropertyPath& InPath) const PURE_VIRTUAL(IStateTreeBindingLookup::GetPropertyPathLeafProperty, return nullptr; );
 };
+
+#if UE_ENABLE_INCLUDE_ORDER_DEPRECATED_IN_5_2
+#include "CoreMinimal.h"
+#endif
diff --git a/Engine/Plugins/Runtime/StateTree/Source/StateTreeModule/Public/StateTreeReference.h b/Engine/Plugins/Runtime/StateTree/Source/StateTreeModule/Public/StateTreeReference.h
--- a/Engine/Plugins/Runtime/StateTree/Source/StateTreeModule/Public/StateTreeReference.h
+++ b/Engine/Plugins/Runtime/StateTree/Source/StateTreeModule/Public/StateTreeReference.h
@@ -2,7 +2,7 @@
 
 #pragma once
 
-#include "StateTreeTypes.h"
+#include "PropertyBag.h"
 #include "StateTreeReference.generated.h"
 
 class UStateTree;
@@ -80,3 +80,7 @@ protected:
 	UPROPERTY(EditAnywhere, Category = "", meta = (FixedLayout))
 	FInstancedPropertyBag Parameters;
 };
+
+#if UE_ENABLE_INCLUDE_ORDER_DEPRECATED_IN_5_2
+#include "StateTreeTypes.h"
+#endif
diff --git a/Engine/Plugins/Runtime/StateTree/Source/StateTreeModule/Public/StateTreeSchema.h b/Engine/Plugins/Runtime/StateTree/Source/StateTreeModule/Public/StateTreeSchema.h
--- a/Engine/Plugins/Runtime/StateTree/Source/StateTreeModule/Public/StateTreeSchema.h
+++ b/Engine/Plugins/Runtime/StateTree/Source/StateTreeModule/Public/StateTreeSchema.h
@@ -2,9 +2,10 @@
 
 #pragma once
 
-#include "StateTreeTypes.h"
 #include "StateTreeSchema.generated.h"
 
+struct FStateTreeExternalDataDesc;
+
 /**
  * Schema describing which inputs, evaluators, and tasks a StateTree can contain.
  * Each StateTree asset saves the schema class name in asset data tags, which can be
@@ -53,3 +54,7 @@ public:
 	
 #endif // WITH_EDITOR
 };
+
+#if UE_ENABLE_INCLUDE_ORDER_DEPRECATED_IN_5_2
+#include "StateTreeTypes.h"
+#endif
diff --git a/Engine/Plugins/Runtime/StateTree/Source/StateTreeModule/Public/StateTreeTaskBase.h b/Engine/Plugins/Runtime/StateTree/Source/StateTreeModule/Public/StateTreeTaskBase.h
--- a/Engine/Plugins/Runtime/StateTree/Source/StateTreeModule/Public/StateTreeTaskBase.h
+++ b/Engine/Plugins/Runtime/StateTree/Source/StateTreeModule/Public/StateTreeTaskBase.h
@@ -2,7 +2,6 @@
 
 #pragma once
 
-#include "StateTreeTypes.h"
 #include "StateTreeNodeBase.h"
 #include "StateTreeTaskBase.generated.h"
 
diff --git a/Engine/Plugins/Runtime/StateTree/Source/StateTreeModule/Public/StateTreeTypes.h b/Engine/Plugins/Runtime/StateTree/Source/StateTreeModule/Public/StateTreeTypes.h
--- a/Engine/Plugins/Runtime/StateTree/Source/StateTreeModule/Public/StateTreeTypes.h
+++ b/Engine/Plugins/Runtime/StateTree/Source/StateTreeModule/Public/StateTreeTypes.h
@@ -2,10 +2,8 @@
 
 #pragma once
 
-#include "StructView.h"
 #include "PropertyBag.h"
 #include "GameplayTagContainer.h"
-#include "Misc/Guid.h"
 #include "StateTreeTypes.generated.h"
 
 STATETREEMODULE_API DECLARE_LOG_CATEGORY_EXTERN(LogStateTree, Warning, All);
diff --git a/Engine/Plugins/Runtime/Synthesis/Source/Synthesis/Classes/SourceEffects/SourceEffectBitCrusher.h b/Engine/Plugins/Runtime/Synthesis/Source/Synthesis/Classes/SourceEffects/SourceEffectBitCrusher.h
--- a/Engine/Plugins/Runtime/Synthesis/Source/Synthesis/Classes/SourceEffects/SourceEffectBitCrusher.h
+++ b/Engine/Plugins/Runtime/Synthesis/Source/Synthesis/Classes/SourceEffects/SourceEffectBitCrusher.h
@@ -1,7 +1,6 @@
 // Copyright Epic Games, Inc. All Rights Reserved.
 #pragma once
 
-#include "Containers/Set.h"
 #include "DSP/BitCrusher.h"
 #include "Sound/SoundEffectSource.h"
 #include "Sound/SoundModulationDestination.h"
diff --git a/Engine/Plugins/Runtime/Synthesis/Source/Synthesis/Classes/SourceEffects/SourceEffectChorus.h b/Engine/Plugins/Runtime/Synthesis/Source/Synthesis/Classes/SourceEffects/SourceEffectChorus.h
--- a/Engine/Plugins/Runtime/Synthesis/Source/Synthesis/Classes/SourceEffects/SourceEffectChorus.h
+++ b/Engine/Plugins/Runtime/Synthesis/Source/Synthesis/Classes/SourceEffects/SourceEffectChorus.h
@@ -2,7 +2,6 @@
 
 #pragma once
 
-#include "Containers/Set.h"
 #include "DSP/Chorus.h"
 #include "Sound/SoundEffectSource.h"
 #include "Sound/SoundModulationDestination.h"
@@ -211,4 +210,4 @@ public:
 
 	UPROPERTY(EditAnywhere, BlueprintReadOnly, Category = "SourceEffect|Preset")
 	FSourceEffectChorusSettings Settings;
-};
\ No newline at end of file
+};
diff --git a/Engine/Plugins/Runtime/Synthesis/Source/Synthesis/Classes/SourceEffects/SourceEffectConvolutionReverb.h b/Engine/Plugins/Runtime/Synthesis/Source/Synthesis/Classes/SourceEffects/SourceEffectConvolutionReverb.h
--- a/Engine/Plugins/Runtime/Synthesis/Source/Synthesis/Classes/SourceEffects/SourceEffectConvolutionReverb.h
+++ b/Engine/Plugins/Runtime/Synthesis/Source/Synthesis/Classes/SourceEffects/SourceEffectConvolutionReverb.h
@@ -1,13 +1,15 @@
 // Copyright Epic Games, Inc. All Rights Reserved.
 #pragma once
 
-#include "CoreMinimal.h"
-#include "UObject/Object.h"
-#include "DSP/Dsp.h"
-#include "EffectConvolutionReverb.h"
 #include "Sound/SoundEffectSource.h"
 #include "SourceEffectConvolutionReverb.generated.h"
 
+class UAudioImpulseResponse;
+class USourceEffectConvolutionReverbPreset;
+enum class ESubmixEffectConvolutionReverbBlockSize : uint8;
+namespace Audio { class FEffectConvolutionReverb; }
+namespace AudioConvReverbIntrinsics { struct FVersionData; }
+
 USTRUCT(BlueprintType)
 struct SYNTHESIS_API FSourceEffectConvolutionReverbSettings
 {
@@ -168,3 +170,9 @@ private:
 #endif
 };
 
+
+#if UE_ENABLE_INCLUDE_ORDER_DEPRECATED_IN_5_2
+#include "CoreMinimal.h"
+#include "DSP/Dsp.h"
+#include "EffectConvolutionReverb.h"
+#endif
diff --git a/Engine/Plugins/Runtime/Synthesis/Source/Synthesis/Classes/SourceEffects/SourceEffectEQ.h b/Engine/Plugins/Runtime/Synthesis/Source/Synthesis/Classes/SourceEffects/SourceEffectEQ.h
--- a/Engine/Plugins/Runtime/Synthesis/Source/Synthesis/Classes/SourceEffects/SourceEffectEQ.h
+++ b/Engine/Plugins/Runtime/Synthesis/Source/Synthesis/Classes/SourceEffects/SourceEffectEQ.h
@@ -3,9 +3,10 @@
 #pragma once
 
 #include "Sound/SoundEffectSource.h"
-#include "DSP/EQ.h"
 #include "SourceEffectEQ.generated.h"
 
+namespace Audio { class FBiquadFilter; }
+
 
 USTRUCT(BlueprintType)
 struct SYNTHESIS_API FSourceEffectEQBand
@@ -88,3 +89,7 @@ public:
 	UPROPERTY(EditAnywhere, BlueprintReadOnly, Category = "Audio|Effects")
 	FSourceEffectEQSettings Settings;
 };
+
+#if UE_ENABLE_INCLUDE_ORDER_DEPRECATED_IN_5_2
+#include "DSP/EQ.h"
+#endif
diff --git a/Engine/Plugins/Runtime/Synthesis/Source/Synthesis/Classes/SourceEffects/SourceEffectEnvelopeFollower.h b/Engine/Plugins/Runtime/Synthesis/Source/Synthesis/Classes/SourceEffects/SourceEffectEnvelopeFollower.h
--- a/Engine/Plugins/Runtime/Synthesis/Source/Synthesis/Classes/SourceEffects/SourceEffectEnvelopeFollower.h
+++ b/Engine/Plugins/Runtime/Synthesis/Source/Synthesis/Classes/SourceEffects/SourceEffectEnvelopeFollower.h
@@ -2,13 +2,13 @@
 
 #pragma once
 
-#include "CoreMinimal.h"
-#include "Tickable.h"
 #include "Sound/SoundEffectSource.h"
 #include "DSP/EnvelopeFollower.h"
 #include "Components/ActorComponent.h"
 #include "SourceEffectEnvelopeFollower.generated.h"
 
+class UEnvelopeFollowerListener;
+
 UENUM(BlueprintType)
 enum class EEnvelopeFollowerPeakMode : uint8
 {
@@ -148,3 +148,8 @@ public:
 	FSourceEffectEnvelopeFollowerSettings Settings;
 };
 
+
+#if UE_ENABLE_INCLUDE_ORDER_DEPRECATED_IN_5_2
+#include "CoreMinimal.h"
+#include "Tickable.h"
+#endif
diff --git a/Engine/Plugins/Runtime/Synthesis/Source/Synthesis/Classes/SourceEffects/SourceEffectMidSideSpreader.h b/Engine/Plugins/Runtime/Synthesis/Source/Synthesis/Classes/SourceEffects/SourceEffectMidSideSpreader.h
--- a/Engine/Plugins/Runtime/Synthesis/Source/Synthesis/Classes/SourceEffects/SourceEffectMidSideSpreader.h
+++ b/Engine/Plugins/Runtime/Synthesis/Source/Synthesis/Classes/SourceEffects/SourceEffectMidSideSpreader.h
@@ -2,10 +2,7 @@
 
 #pragma once
 
-#include "CoreMinimal.h"
 #include "Sound/SoundEffectSource.h"
-#include "UObject/ObjectMacros.h"
-#include "DSP/Dsp.h"
 #include "SourceEffectMidSideSpreader.generated.h"
 
 // Stereo channel mode
@@ -108,3 +105,8 @@ public:
 	UPROPERTY(EditAnywhere, BlueprintReadOnly, Category = "Audio|Effects")
 	FSourceEffectMidSideSpreaderSettings Settings;
 };
+
+#if UE_ENABLE_INCLUDE_ORDER_DEPRECATED_IN_5_2
+#include "CoreMinimal.h"
+#include "DSP/Dsp.h"
+#endif
diff --git a/Engine/Plugins/Runtime/Synthesis/Source/Synthesis/Classes/SourceEffects/SourceEffectMotionFilter.h b/Engine/Plugins/Runtime/Synthesis/Source/Synthesis/Classes/SourceEffects/SourceEffectMotionFilter.h
--- a/Engine/Plugins/Runtime/Synthesis/Source/Synthesis/Classes/SourceEffects/SourceEffectMotionFilter.h
+++ b/Engine/Plugins/Runtime/Synthesis/Source/Synthesis/Classes/SourceEffects/SourceEffectMotionFilter.h
@@ -2,13 +2,12 @@
 
 #pragma once
 
-#include "CoreMinimal.h"
 #include "Sound/SoundEffectSource.h"
-#include "DSP/Dsp.h"
 #include "DSP/Filter.h"
-#include "DSP/BufferVectorOperations.h"
 #include "SourceEffectMotionFilter.generated.h"
 
+namespace Audio { class FLinearEase; }
+
 UENUM(BlueprintType)
 enum class ESourceEffectMotionFilterModSource : uint8
 {
@@ -361,3 +360,8 @@ public:
 	UPROPERTY(EditAnywhere, BlueprintReadOnly, Category = "Audio|Effects")
 	FSourceEffectMotionFilterSettings Settings;
 };
+
+#if UE_ENABLE_INCLUDE_ORDER_DEPRECATED_IN_5_2
+#include "CoreMinimal.h"
+#include "DSP/BufferVectorOperations.h"
+#endif
diff --git a/Engine/Plugins/Runtime/Synthesis/Source/Synthesis/Classes/SubmixEffects/SubmixEffectConvolutionReverb.h b/Engine/Plugins/Runtime/Synthesis/Source/Synthesis/Classes/SubmixEffects/SubmixEffectConvolutionReverb.h
--- a/Engine/Plugins/Runtime/Synthesis/Source/Synthesis/Classes/SubmixEffects/SubmixEffectConvolutionReverb.h
+++ b/Engine/Plugins/Runtime/Synthesis/Source/Synthesis/Classes/SubmixEffects/SubmixEffectConvolutionReverb.h
@@ -1,13 +1,15 @@
 // Copyright Epic Games, Inc. All Rights Reserved.
 #pragma once
 
-#include "CoreMinimal.h"
-#include "UObject/Object.h"
-#include "DSP/Dsp.h"
-#include "EffectConvolutionReverb.h"
 #include "Sound/SoundEffectSubmix.h"
 #include "SubmixEffectConvolutionReverb.generated.h"
 
+class UAudioImpulseResponse;
+class USubmixEffectConvolutionReverbPreset;
+enum class ESubmixEffectConvolutionReverbBlockSize : uint8;
+namespace Audio { class FEffectConvolutionReverb; }
+namespace AudioConvReverbIntrinsics { struct FVersionData; }
+
 USTRUCT(BlueprintType)
 struct SYNTHESIS_API FSubmixEffectConvolutionReverbSettings
 {
@@ -204,3 +206,9 @@ private:
 #endif
 };
 
+
+#if UE_ENABLE_INCLUDE_ORDER_DEPRECATED_IN_5_2
+#include "CoreMinimal.h"
+#include "DSP/Dsp.h"
+#include "EffectConvolutionReverb.h"
+#endif
diff --git a/Engine/Plugins/Runtime/Synthesis/Source/Synthesis/Classes/SubmixEffects/SubmixEffectDelay.h b/Engine/Plugins/Runtime/Synthesis/Source/Synthesis/Classes/SubmixEffects/SubmixEffectDelay.h
--- a/Engine/Plugins/Runtime/Synthesis/Source/Synthesis/Classes/SubmixEffects/SubmixEffectDelay.h
+++ b/Engine/Plugins/Runtime/Synthesis/Source/Synthesis/Classes/SubmixEffects/SubmixEffectDelay.h
@@ -2,14 +2,13 @@
 
 #pragma once
 
-#include "CoreMinimal.h"
-#include "AudioEffect.h"
-#include "DSP/Delay.h"
 #include "DSP/Dsp.h"
 #include "Kismet/BlueprintFunctionLibrary.h"
 #include "Sound/SoundEffectSubmix.h"
 #include "SubmixEffectDelay.generated.h"
 
+namespace Audio { class FDelay; }
+
 // ========================================================================
 // FSubmixEffectDelaySettings
 // UStruct used to define user-exposed params for use with your effect.
@@ -172,3 +171,9 @@ public:
 	UPROPERTY(transient)
 	FSubmixEffectDelaySettings DynamicSettings;
 };
+
+#if UE_ENABLE_INCLUDE_ORDER_DEPRECATED_IN_5_2
+#include "AudioEffect.h"
+#include "CoreMinimal.h"
+#include "DSP/Delay.h"
+#endif
diff --git a/Engine/Plugins/Runtime/Synthesis/Source/Synthesis/Classes/SubmixEffects/SubmixEffectFilter.h b/Engine/Plugins/Runtime/Synthesis/Source/Synthesis/Classes/SubmixEffects/SubmixEffectFilter.h
--- a/Engine/Plugins/Runtime/Synthesis/Source/Synthesis/Classes/SubmixEffects/SubmixEffectFilter.h
+++ b/Engine/Plugins/Runtime/Synthesis/Source/Synthesis/Classes/SubmixEffects/SubmixEffectFilter.h
@@ -2,8 +2,6 @@
 
 #pragma once
 
-#include "CoreMinimal.h"
-#include "AudioEffect.h"
 #include "DSP/Filter.h"
 #include "Sound/SoundEffectSubmix.h"
 #include "SubmixEffectFilter.generated.h"
@@ -169,3 +167,8 @@ public:
 	UPROPERTY(EditAnywhere, BlueprintReadWrite, Category = SubmixEffectPreset, Meta = (ShowOnlyInnerProperties))
 	FSubmixEffectFilterSettings Settings;
 };
+
+#if UE_ENABLE_INCLUDE_ORDER_DEPRECATED_IN_5_2
+#include "AudioEffect.h"
+#include "CoreMinimal.h"
+#endif
diff --git a/Engine/Plugins/Runtime/Synthesis/Source/Synthesis/Classes/SubmixEffects/SubmixEffectFlexiverb.h b/Engine/Plugins/Runtime/Synthesis/Source/Synthesis/Classes/SubmixEffects/SubmixEffectFlexiverb.h
--- a/Engine/Plugins/Runtime/Synthesis/Source/Synthesis/Classes/SubmixEffects/SubmixEffectFlexiverb.h
+++ b/Engine/Plugins/Runtime/Synthesis/Source/Synthesis/Classes/SubmixEffects/SubmixEffectFlexiverb.h
@@ -4,7 +4,6 @@
 
 #include "Flexiverb.h"
 #include "Sound/SoundEffectSubmix.h"
-#include "AudioEffect.h"
 #include "SubmixEffectFlexiverb.generated.h"
 
 USTRUCT(BlueprintType)
@@ -81,4 +80,8 @@ public:
 
 	UPROPERTY(EditAnywhere, BlueprintReadWrite, Category = SubmixEffectPreset)
 	FSubmixEffectFlexiverbSettings Settings;
-};
\ No newline at end of file
+};
+
+#if UE_ENABLE_INCLUDE_ORDER_DEPRECATED_IN_5_2
+#include "AudioEffect.h"
+#endif
diff --git a/Engine/Plugins/Runtime/Synthesis/Source/Synthesis/Classes/SubmixEffects/SubmixEffectMultiBandCompressor.h b/Engine/Plugins/Runtime/Synthesis/Source/Synthesis/Classes/SubmixEffects/SubmixEffectMultiBandCompressor.h
--- a/Engine/Plugins/Runtime/Synthesis/Source/Synthesis/Classes/SubmixEffects/SubmixEffectMultiBandCompressor.h
+++ b/Engine/Plugins/Runtime/Synthesis/Source/Synthesis/Classes/SubmixEffects/SubmixEffectMultiBandCompressor.h
@@ -2,12 +2,12 @@
 
 #pragma once
 
-#include "DSP/DynamicsProcessor.h"
-#include "Sound/SoundEffectSubmix.h"
 #include "DSP/LinkwitzRileyBandSplitter.h"
 #include "SubmixEffects/AudioMixerSubmixEffectDynamicsProcessor.h"
 #include "SubmixEffectMultiBandCompressor.generated.h"
 
+namespace Audio { class FDynamicsProcessor; }
+
 USTRUCT(BlueprintType)
 struct SYNTHESIS_API FDynamicsBandSettings
 {
