[AudioWidgets]
* Fixed bad auto merge

diff --git a/Engine/Plugins/Runtime/AudioWidgets/Source/AudioWidgets/Public/AudioMeter.h b/Engine/Plugins/Runtime/AudioWidgets/Source/AudioWidgets/Public/AudioMeter.h
--- a/Engine/Plugins/Runtime/AudioWidgets/Source/AudioWidgets/Public/AudioMeter.h
+++ b/Engine/Plugins/Runtime/AudioWidgets/Source/AudioWidgets/Public/AudioMeter.h
@@ -150,15 +150,13 @@ protected:
 	PROPERTY_BINDING_IMPLEMENTATION(TArray<FMeterChannelInfo>, MeterChannelInfo);
 };
 
-#if UE_ENABLE_INCLUDE_ORDER_DEPRECATED_IN_5_2
-#include "CoreMinimal.h"
-#include "Styling/SlateTypes.h"
+namespace AudioWidgets
 {
 	class AUDIOWIDGETS_API FAudioMeter
 	{
 	public:
 		FAudioMeter();
-#endif
+
 		UAudioBus* GetAudioBus() const;
 
 		TSharedRef<SAudioMeter> GetWidget() const;
