Fixed compile issues found after moving module files around in the unity files.
Also added missing generated.h includes.

diff --git a/Engine/Plugins/Runtime/AudioWidgets/Source/AudioWidgets/Public/SAudioRadialSlider.h b/Engine/Plugins/Runtime/AudioWidgets/Source/AudioWidgets/Public/SAudioRadialSlider.h
--- a/Engine/Plugins/Runtime/AudioWidgets/Source/AudioWidgets/Public/SAudioRadialSlider.h
+++ b/Engine/Plugins/Runtime/AudioWidgets/Source/AudioWidgets/Public/SAudioRadialSlider.h
@@ -10,6 +10,8 @@
 #include "Styling/SlateStyleRegistry.h"
 #include "Styling/SlateWidgetStyleAsset.h"
 
+#include "SAudioRadialSlider.generated.h"
+
 class SAudioTextBox;
 class SImage;
 class SWidgetSwitcher;
