Fixed compile issue after merge.

diff --git a/Engine/Plugins/Editor/DataValidation/Source/DataValidation/Public/DataValidationModule.h b/Engine/Plugins/Editor/DataValidation/Source/DataValidation/Public/DataValidationModule.h
--- a/Engine/Plugins/Editor/DataValidation/Source/DataValidation/Public/DataValidationModule.h
+++ b/Engine/Plugins/Editor/DataValidation/Source/DataValidation/Public/DataValidationModule.h
@@ -5,12 +5,15 @@
 #include "Modules/ModuleManager.h"
 #include "UObject/ObjectMacros.h"
 
+#if UE_ENABLE_INCLUDE_ORDER_DEPRECATED_IN_5_2
+#include "AssetRegistry/AssetData.h"
+#include "CoreMinimal.h"
+#endif
+
 #include "DataValidationModule.generated.h"
 
 struct FAssetData;
 
-#include "DataValidationModule.generated.h"
-
 UENUM(BlueprintType)
 enum class EDataValidationUsecase : uint8
 {
@@ -68,8 +71,3 @@ public:
 	/** Validates selected assets and opens a window to report the results. If bValidateDependencies is true it will also validate any assets that the selected assets depend on. */
 	virtual void ValidateAssets(const TArray<FAssetData>& SelectedAssets, bool bValidateDependencies, const EDataValidationUsecase InValidationUsecase) = 0;
 };
-
-#if UE_ENABLE_INCLUDE_ORDER_DEPRECATED_IN_5_2
-#include "AssetRegistry/AssetData.h"
-#include "CoreMinimal.h"
-#endif
diff --git a/Engine/Plugins/Runtime/AudioWidgets/Source/AudioWidgets/Public/SAudioRadialSlider.h b/Engine/Plugins/Runtime/AudioWidgets/Source/AudioWidgets/Public/SAudioRadialSlider.h
--- a/Engine/Plugins/Runtime/AudioWidgets/Source/AudioWidgets/Public/SAudioRadialSlider.h
+++ b/Engine/Plugins/Runtime/AudioWidgets/Source/AudioWidgets/Public/SAudioRadialSlider.h
@@ -10,14 +10,21 @@
 #include "Styling/SlateStyleRegistry.h"
 #include "Styling/SlateWidgetStyleAsset.h"
 
+#if UE_ENABLE_INCLUDE_ORDER_DEPRECATED_IN_5_2
+#include "AudioDefines.h"
+#include "SAudioTextBox.h"
+#include "Styling/CoreStyle.h"
+#include "Widgets/Images/SImage.h"
+#include "Widgets/Input/SEditableText.h"
+#include "Widgets/Layout/SWidgetSwitcher.h"
+#include "Widgets/SOverlay.h"
+#endif
+
 #include "SAudioRadialSlider.generated.h"
 
 class SAudioTextBox;
 class SImage;
 class SWidgetSwitcher;
-
-#include "SAudioRadialSlider.generated.h"
-
 class SRadialSlider;
 
 UENUM()
@@ -198,13 +205,3 @@ public:
 	const float GetOutputValue(const float InSliderValue);
 	const float GetSliderValue(const float OutputValue);
 };
-
-#if UE_ENABLE_INCLUDE_ORDER_DEPRECATED_IN_5_2
-#include "AudioDefines.h"
-#include "SAudioTextBox.h"
-#include "Styling/CoreStyle.h"
-#include "Widgets/Images/SImage.h"
-#include "Widgets/Input/SEditableText.h"
-#include "Widgets/Layout/SWidgetSwitcher.h"
-#include "Widgets/SOverlay.h"
-#endif
