[AudioWidgets] Update naming convention to distinguish between sparce and fixed sample sequences

#trivial

diff --git a/Engine/Plugins/Editor/WaveformEditor/Source/WaveformEditor/Private/WaveformEditor.cpp b/Engine/Plugins/Editor/WaveformEditor/Source/WaveformEditor/Private/WaveformEditor.cpp
--- a/Engine/Plugins/Editor/WaveformEditor/Source/WaveformEditor/Private/WaveformEditor.cpp
+++ b/Engine/Plugins/Editor/WaveformEditor/Source/WaveformEditor/Private/WaveformEditor.cpp
@@ -6,7 +6,7 @@
 #include "Components/AudioComponent.h"
 #include "Misc/TransactionObjectEvent.h"
 #include "PropertyEditorModule.h"
-#include "SamplesSequenceTransportCoordinator.h"
+#include "SparseSampledSequenceTransportCoordinator.h"
 #include "Sound/SoundWave.h"
 #include "Styling/AppStyle.h"
 #include "SWaveformPanel.h"
@@ -629,7 +629,7 @@ bool FWaveformEditor::SetUpWaveformPanel()
 
 	WaveformRenderData = MakeShared<FWaveformEditorRenderData>();
 	
-	TransportCoordinator = MakeShared<FSamplesSequenceTransportCoordinator>();
+	TransportCoordinator = MakeShared<FSparseSampledSequenceTransportCoordinator>();
 	WaveformRenderData->OnRenderDataUpdated.AddSP(this, &FWaveformEditor::HandleRenderDataUpdate);
 	
 	FOnTransformationsPropertiesRequired OnTransformationPropertiesRequired = FOnTransformationsPropertiesRequired::CreateLambda([this](FTransformationsToPropertiesArray& InObjToPropsMap)
diff --git a/Engine/Plugins/Editor/WaveformEditor/Source/WaveformEditor/Public/WaveformEditor.h b/Engine/Plugins/Editor/WaveformEditor/Source/WaveformEditor/Public/WaveformEditor.h
--- a/Engine/Plugins/Editor/WaveformEditor/Source/WaveformEditor/Public/WaveformEditor.h
+++ b/Engine/Plugins/Editor/WaveformEditor/Source/WaveformEditor/Public/WaveformEditor.h
@@ -104,7 +104,7 @@ private:
 	TSharedPtr<class FWaveformEditorWaveWriter> WaveWriter = nullptr;
 
 	/** Manages Transport info in waveform panel */
-	TSharedPtr<class FSamplesSequenceTransportCoordinator> TransportCoordinator = nullptr;
+	TSharedPtr<class FSparseSampledSequenceTransportCoordinator> TransportCoordinator = nullptr;
 
 	/** Controls Transport of the audio component  */
 	TSharedPtr<FWaveformEditorTransportController> TransportController = nullptr;
diff --git a/Engine/Plugins/Editor/WaveformEditor/Source/WaveformEditorWidgets/Private/SWaveformEditorTimeRuler.cpp b/Engine/Plugins/Editor/WaveformEditor/Source/WaveformEditorWidgets/Private/SWaveformEditorTimeRuler.cpp
--- a/Engine/Plugins/Editor/WaveformEditor/Source/WaveformEditorWidgets/Private/SWaveformEditorTimeRuler.cpp
+++ b/Engine/Plugins/Editor/WaveformEditor/Source/WaveformEditorWidgets/Private/SWaveformEditorTimeRuler.cpp
@@ -5,11 +5,11 @@
 #include "Layout/WidgetPath.h"
 #include "Framework/Application/SlateApplication.h"
 #include "Framework/MultiBox/MultiBoxBuilder.h"
-#include "SamplesSequenceTransportCoordinator.h"
+#include "SparseSampledSequenceTransportCoordinator.h"
 
 #define LOCTEXT_NAMESPACE "WaveformEditorTimeRuler"
 
-void SWaveformEditorTimeRuler::Construct(const FArguments& InArgs, TSharedRef<ISampledSequenceGridService> InGridService)
+void SWaveformEditorTimeRuler::Construct(const FArguments& InArgs, TSharedRef<IFixedSampledSequenceGridService> InGridService)
 {
 	GridService = InGridService;
 	UpdateGridMetrics();
diff --git a/Engine/Plugins/Editor/WaveformEditor/Source/WaveformEditorWidgets/Private/SWaveformEditorTimeRuler.h b/Engine/Plugins/Editor/WaveformEditor/Source/WaveformEditorWidgets/Private/SWaveformEditorTimeRuler.h
--- a/Engine/Plugins/Editor/WaveformEditor/Source/WaveformEditorWidgets/Private/SWaveformEditorTimeRuler.h
+++ b/Engine/Plugins/Editor/WaveformEditor/Source/WaveformEditorWidgets/Private/SWaveformEditorTimeRuler.h
@@ -2,7 +2,7 @@
 
 #pragma once
 
-#include "ISampledSequenceGridService.h"
+#include "IFixedSampledSequenceGridService.h"
 #include "Framework/MultiBox/MultiBoxBuilder.h"
 #include "Styling/SlateStyleRegistry.h"
 #include "Styling/SlateWidgetStyleAsset.h"
@@ -12,8 +12,6 @@
 
 DECLARE_MULTICAST_DELEGATE_OneParam(FOnTimeUnitMenuSelection, const EWaveformEditorDisplayUnit /* Requested Display Unit */);
 
-class FSamplesSequenceTransportCoordinator;
-
 class SWaveformEditorTimeRuler : public SCompoundWidget
 {
 
@@ -29,7 +27,7 @@ class SWaveformEditorTimeRuler : public SCompoundWidget
 	SLATE_END_ARGS()
 
 public:
-	void Construct(const FArguments& InArgs, TSharedRef<ISampledSequenceGridService> InGridService);
+	void Construct(const FArguments& InArgs, TSharedRef<IFixedSampledSequenceGridService> InGridService);
 	void UpdateGridMetrics();
 	void UpdateDisplayUnit(const EWaveformEditorDisplayUnit InDisplayUnit);
 	void SetPlayheadPosition(const float InNewPosition);
@@ -51,7 +49,7 @@ private:
 	void MakeTimeUnitsSubMenu(FMenuBuilder& SubMenuBuilder);
 	void NotifyTimeUnitMenuSelection(const EWaveformEditorDisplayUnit SelectedDisplayUnit) const;
 
-	FSampledSequenceGridMetrics GridMetrics;
+	FFixedSampledSequenceGridMetrics GridMetrics;
 
 	const FWaveformEditorTimeRulerStyle* Style = nullptr;
 
@@ -68,7 +66,7 @@ private:
 	float PlayheadPosition = 0.f;
 	FSlateFontInfo TicksTextFont;
 
-	TSharedPtr<ISampledSequenceGridService> GridService = nullptr;
+	TSharedPtr<IFixedSampledSequenceGridService> GridService = nullptr;
 
 	EWaveformEditorDisplayUnit DisplayUnit;
 };
diff --git a/Engine/Plugins/Editor/WaveformEditor/Source/WaveformEditorWidgets/Private/SWaveformPanel.cpp b/Engine/Plugins/Editor/WaveformEditor/Source/WaveformEditorWidgets/Private/SWaveformPanel.cpp
--- a/Engine/Plugins/Editor/WaveformEditor/Source/WaveformEditorWidgets/Private/SWaveformPanel.cpp
+++ b/Engine/Plugins/Editor/WaveformEditor/Source/WaveformEditorWidgets/Private/SWaveformPanel.cpp
@@ -3,11 +3,11 @@
 #include "SWaveformPanel.h"
 
 #include "SPlayheadOverlay.h"
-#include "SSampledSequenceViewer.h"
+#include "SFixedSampledSequenceViewer.h"
 #include "SWaveformEditorTimeRuler.h"
 #include "SWaveformTransformationsOverlay.h"
 #include "SWaveformViewerOverlay.h"
-#include "SamplesSequenceTransportCoordinator.h"
+#include "SparseSampledSequenceTransportCoordinator.h"
 #include "WaveformEditorDisplayUnit.h"
 #include "WaveformEditorGridData.h"
 #include "WaveformEditorRenderData.h"
@@ -17,7 +17,7 @@
 #include "Widgets/SCompoundWidget.h"
 #include "Widgets/SOverlay.h"
 
-void SWaveformPanel::Construct(const FArguments& InArgs, TSharedRef<FWaveformEditorRenderData> InRenderData, TSharedRef<FSamplesSequenceTransportCoordinator> InTransportCoordinator, TSharedRef<FWaveformEditorZoomController> InZoomManager, TSharedPtr<SWaveformTransformationsOverlay> InWaveformTransformationsOverlay)
+void SWaveformPanel::Construct(const FArguments& InArgs, TSharedRef<FWaveformEditorRenderData> InRenderData, TSharedRef<FSparseSampledSequenceTransportCoordinator> InTransportCoordinator, TSharedRef<FWaveformEditorZoomController> InZoomManager, TSharedPtr<SWaveformTransformationsOverlay> InWaveformTransformationsOverlay)
 {
 	DisplayUnit = EWaveformEditorDisplayUnit::Seconds;
 
@@ -143,12 +143,12 @@ void SWaveformPanel::SetUpWaveformViewer(TSharedRef<FWaveformEditorGridData> InG
 	FSampledSequenceViewerStyle* WaveViewerStyle = &WaveformEditorStyle->GetRegisteredWidgetStyle<FSampledSequenceViewerStyle>("WaveformViewer.Style").Get();
 	check(WaveViewerStyle);
 	
-	TimeSeriesDrawingUtils::FSampledSequenceDrawingParams WaveformViewerDrawingParams;
+	SampledSequenceDrawingUtils::FSampledSequenceDrawingParams WaveformViewerDrawingParams;
 	WaveformViewerDrawingParams.MaxDisplayedValue = TNumericLimits<int16>::Max();
 
-	WaveformViewer = SNew(SSampledSequenceViewer, MakeArrayView(FloatRenderData.GetData(), FloatRenderData.Num()), InRenderData->GetNumChannels(), InGridData).Style(WaveViewerStyle).SequenceDrawingParams(WaveformViewerDrawingParams);
-	WaveViewerStyle->OnStyleUpdated.AddSP(WaveformViewer.ToSharedRef(), &SSampledSequenceViewer::OnStyleUpdated);
-	GridData->OnGridMetricsUpdated.AddSP(WaveformViewer.ToSharedRef(), &SSampledSequenceViewer::UpdateGridMetrics);
+	WaveformViewer = SNew(SFixedSampledSequenceViewer, MakeArrayView(FloatRenderData.GetData(), FloatRenderData.Num()), InRenderData->GetNumChannels(), InGridData).Style(WaveViewerStyle).SequenceDrawingParams(WaveformViewerDrawingParams);
+	WaveViewerStyle->OnStyleUpdated.AddSP(WaveformViewer.ToSharedRef(), &SFixedSampledSequenceViewer::OnStyleUpdated);
+	GridData->OnGridMetricsUpdated.AddSP(WaveformViewer.ToSharedRef(), &SFixedSampledSequenceViewer::UpdateGridMetrics);
 }
 
 void SWaveformPanel::SetUpGridData(TSharedRef<FWaveformEditorRenderData> InRenderData)
@@ -159,9 +159,9 @@ void SWaveformPanel::SetUpGridData(TSharedRef<FWaveformEditorRenderData> InRende
 	GridData = MakeShared<FWaveformEditorGridData>(InRenderData->GetNumSamples() / InRenderData->GetNumChannels(), InRenderData->GetSampleRate(), RulerStyle->DesiredWidth, &RulerStyle->TicksTextFont);
 }
 
-void SWaveformPanel::SetUpZoomManager(TSharedRef<FWaveformEditorZoomController> InZoomManager, TSharedRef<FSamplesSequenceTransportCoordinator> InTransportCoordinator)
+void SWaveformPanel::SetUpZoomManager(TSharedRef<FWaveformEditorZoomController> InZoomManager, TSharedRef<FSparseSampledSequenceTransportCoordinator> InTransportCoordinator)
 {
-	InZoomManager->OnZoomRatioChanged.AddSP(InTransportCoordinator, &FSamplesSequenceTransportCoordinator::SetZoomRatio);
+	InZoomManager->OnZoomRatioChanged.AddSP(InTransportCoordinator, &FSparseSampledSequenceTransportCoordinator::SetZoomRatio);
 }
 
 void SWaveformPanel::OnRenderDataUpdated()
diff --git a/Engine/Plugins/Editor/WaveformEditor/Source/WaveformEditorWidgets/Private/SWaveformTransformationsOverlay.cpp b/Engine/Plugins/Editor/WaveformEditor/Source/WaveformEditorWidgets/Private/SWaveformTransformationsOverlay.cpp
--- a/Engine/Plugins/Editor/WaveformEditor/Source/WaveformEditorWidgets/Private/SWaveformTransformationsOverlay.cpp
+++ b/Engine/Plugins/Editor/WaveformEditor/Source/WaveformEditorWidgets/Private/SWaveformTransformationsOverlay.cpp
@@ -4,10 +4,10 @@
 
 #include "AudioWidgetsUtils.h"
 #include "SWaveformTransformationRenderLayer.h"
-#include "SamplesSequenceTransportCoordinator.h"
+#include "SparseSampledSequenceTransportCoordinator.h"
 #include "Widgets/SOverlay.h"
 
-void SWaveformTransformationsOverlay::Construct(const FArguments& InArgs, TArrayView<const FTransformationLayerRenderInfo> InTransformationRenderers, TSharedRef<FSamplesSequenceTransportCoordinator> InTransportCoordinator)
+void SWaveformTransformationsOverlay::Construct(const FArguments& InArgs, TArrayView<const FTransformationLayerRenderInfo> InTransformationRenderers, TSharedRef<FSparseSampledSequenceTransportCoordinator> InTransportCoordinator)
 {
 	TransformationRenderers = InTransformationRenderers;
 	TransportCoordinator = InTransportCoordinator;
diff --git a/Engine/Plugins/Editor/WaveformEditor/Source/WaveformEditorWidgets/Private/SWaveformViewerOverlay.h b/Engine/Plugins/Editor/WaveformEditor/Source/WaveformEditorWidgets/Private/SWaveformViewerOverlay.h
--- a/Engine/Plugins/Editor/WaveformEditor/Source/WaveformEditorWidgets/Private/SWaveformViewerOverlay.h
+++ b/Engine/Plugins/Editor/WaveformEditor/Source/WaveformEditorWidgets/Private/SWaveformViewerOverlay.h
@@ -10,7 +10,7 @@
 DECLARE_DELEGATE_OneParam(FOnNewMouseDelta, const float /* new delta */)
 
 class FWaveformEditorGridData;
-class FSamplesSequenceTransportCoordinator;
+class FSparseSampledSequenceTransportCoordinator;
 
 class SWaveformViewerOverlay : public SLeafWidget
 {
diff --git a/Engine/Plugins/Editor/WaveformEditor/Source/WaveformEditorWidgets/Private/WaveformEditorGridData.cpp b/Engine/Plugins/Editor/WaveformEditor/Source/WaveformEditorWidgets/Private/WaveformEditorGridData.cpp
--- a/Engine/Plugins/Editor/WaveformEditor/Source/WaveformEditorWidgets/Private/WaveformEditorGridData.cpp
+++ b/Engine/Plugins/Editor/WaveformEditor/Source/WaveformEditorWidgets/Private/WaveformEditorGridData.cpp
@@ -67,7 +67,7 @@ bool FWaveformEditorGridData::UpdateGridMetrics(const float InGridSizePixels)
 	return true;
 }
 
-const FSampledSequenceGridMetrics FWaveformEditorGridData::GetGridMetrics() const
+const FFixedSampledSequenceGridMetrics FWaveformEditorGridData::GetGridMetrics() const
 {
 	return GridMetrics;
 }
diff --git a/Engine/Plugins/Editor/WaveformEditor/Source/WaveformEditorWidgets/Public/SWaveformPanel.h b/Engine/Plugins/Editor/WaveformEditor/Source/WaveformEditorWidgets/Public/SWaveformPanel.h
--- a/Engine/Plugins/Editor/WaveformEditor/Source/WaveformEditorWidgets/Public/SWaveformPanel.h
+++ b/Engine/Plugins/Editor/WaveformEditor/Source/WaveformEditorWidgets/Public/SWaveformPanel.h
@@ -6,13 +6,13 @@
 #include "Widgets/SCompoundWidget.h"
 
 enum class EWaveformEditorDisplayUnit;
-class FSamplesSequenceTransportCoordinator;
+class FSparseSampledSequenceTransportCoordinator;
 class FWaveformEditorGridData;
 class FWaveformEditorRenderData;
 class FWaveformEditorStyle;
 class FWaveformEditorZoomController;
 class SPlayheadOverlay;
-class SSampledSequenceViewer;
+class SFixedSampledSequenceViewer;
 class SWaveformEditorTimeRuler;
 class SWaveformTransformationsOverlay;
 class SWaveformViewerOverlay;
@@ -35,7 +35,7 @@ public:
 		SLATE_DEFAULT_SLOT(FArguments, InArgs)
 	SLATE_END_ARGS()
 
-	void Construct(const FArguments& InArgs, TSharedRef<FWaveformEditorRenderData> InRenderData, TSharedRef<FSamplesSequenceTransportCoordinator> InTransportCoordinator, TSharedRef<FWaveformEditorZoomController> InZoomManager, TSharedPtr<SWaveformTransformationsOverlay> InWaveformTransformationsOverlay = nullptr);
+	void Construct(const FArguments& InArgs, TSharedRef<FWaveformEditorRenderData> InRenderData, TSharedRef<FSparseSampledSequenceTransportCoordinator> InTransportCoordinator, TSharedRef<FWaveformEditorZoomController> InZoomManager, TSharedPtr<SWaveformTransformationsOverlay> InWaveformTransformationsOverlay = nullptr);
 	virtual void Tick(const FGeometry& AllottedGeometry, const double InCurrentTime, const float InDeltaTime) override;
 
 private:
@@ -46,7 +46,7 @@ private:
 	void SetUpTimeRuler(TSharedRef<FWaveformEditorGridData> InGridData);
 	void SetUpWaveformViewer(TSharedRef<FWaveformEditorGridData> InGridData, TSharedRef<FWaveformEditorRenderData> InRenderData);
 	void SetUpWaveformViewerOverlay(TSharedRef<FWaveformEditorZoomController> InZoomManager);
-	void SetUpZoomManager(TSharedRef<FWaveformEditorZoomController> InZoomManager, TSharedRef<FSamplesSequenceTransportCoordinator> InTransportCoordinator);
+	void SetUpZoomManager(TSharedRef<FWaveformEditorZoomController> InZoomManager, TSharedRef<FSparseSampledSequenceTransportCoordinator> InTransportCoordinator);
 
 	void OnRenderDataUpdated();
 	void OnDisplayRangeUpdated(const TRange<float> NewDisplayRange);
@@ -67,10 +67,10 @@ private:
 
 	TSharedPtr<FWaveformEditorRenderData> RenderData;
 	TSharedPtr<FWaveformEditorGridData> GridData;
-	TSharedPtr<FSamplesSequenceTransportCoordinator> TransportCoordinator;
+	TSharedPtr<FSparseSampledSequenceTransportCoordinator> TransportCoordinator;
 
 	TSharedPtr<SWaveformEditorTimeRuler> TimeRuler;
-	TSharedPtr<SSampledSequenceViewer> WaveformViewer;
+	TSharedPtr<SFixedSampledSequenceViewer> WaveformViewer;
 	TSharedPtr<SWaveformTransformationsOverlay> WaveformTransformationsOverlay;
 	TSharedPtr<SWaveformViewerOverlay> WaveformViewerOverlay;
 	TSharedPtr<SPlayheadOverlay> PlayheadOverlay;
diff --git a/Engine/Plugins/Editor/WaveformEditor/Source/WaveformEditorWidgets/Public/SWaveformTransformationsOverlay.h b/Engine/Plugins/Editor/WaveformEditor/Source/WaveformEditorWidgets/Public/SWaveformTransformationsOverlay.h
--- a/Engine/Plugins/Editor/WaveformEditor/Source/WaveformEditorWidgets/Public/SWaveformTransformationsOverlay.h
+++ b/Engine/Plugins/Editor/WaveformEditor/Source/WaveformEditorWidgets/Public/SWaveformTransformationsOverlay.h
@@ -5,7 +5,7 @@
 #include "Widgets/SCompoundWidget.h"
 #include "Widgets/Layout/SConstraintCanvas.h"
 
-class FSamplesSequenceTransportCoordinator;
+class FSparseSampledSequenceTransportCoordinator;
 class IWaveformTransformationRenderer;
 class SOverlay;
 class SWaveformTransformationRenderLayer;
@@ -20,7 +20,7 @@ public:
 		SLATE_DEFAULT_SLOT(FArguments, InArgs)
 	SLATE_END_ARGS()
 
-	void Construct(const FArguments& InArgs, TArrayView< const FTransformationLayerRenderInfo> InTransformationRenderers, TSharedRef<FSamplesSequenceTransportCoordinator> InTransportCoordinator);
+	void Construct(const FArguments& InArgs, TArrayView< const FTransformationLayerRenderInfo> InTransformationRenderers, TSharedRef<FSparseSampledSequenceTransportCoordinator> InTransportCoordinator);
 	void OnLayerChainGenerated(FTransformationLayerRenderInfo* FirstLayerPtr, const int32 NLayers);
 	void UpdateLayerConstraints();
 	void OnNewWaveformDisplayRange(const TRange<float> NewDisplayRange);
@@ -39,5 +39,5 @@ private:
 	TArray<TSharedPtr<SWidget>> TransformationLayers;
 	TArray<SConstraintCanvas::FSlot*> LayersSlots;
 	TArrayView<const FTransformationLayerRenderInfo> TransformationRenderers;
-	TSharedPtr<FSamplesSequenceTransportCoordinator> TransportCoordinator = nullptr;
+	TSharedPtr<FSparseSampledSequenceTransportCoordinator> TransportCoordinator = nullptr;
 };
\ No newline at end of file
diff --git a/Engine/Plugins/Editor/WaveformEditor/Source/WaveformEditorWidgets/Public/WaveformEditorGridData.h b/Engine/Plugins/Editor/WaveformEditor/Source/WaveformEditorWidgets/Public/WaveformEditorGridData.h
--- a/Engine/Plugins/Editor/WaveformEditor/Source/WaveformEditorWidgets/Public/WaveformEditorGridData.h
+++ b/Engine/Plugins/Editor/WaveformEditor/Source/WaveformEditorWidgets/Public/WaveformEditorGridData.h
@@ -4,7 +4,7 @@
 
 #include "Delegates/Delegate.h"
 #include "HAL/Platform.h"
-#include "ISampledSequenceGridService.h"
+#include "IFixedSampledSequenceGridService.h"
 #include "Math/Range.h"
 #include "Misc/FrameRate.h"
 
@@ -12,21 +12,21 @@ DECLARE_MULTICAST_DELEGATE(FOnWaveformEditorGridUpdated)
 
 struct FSlateFontInfo;
 
-class WAVEFORMEDITORWIDGETS_API FWaveformEditorGridData : public ISampledSequenceGridService
+class WAVEFORMEDITORWIDGETS_API FWaveformEditorGridData : public IFixedSampledSequenceGridService
 {
 public: 
 	explicit FWaveformEditorGridData(const uint32 InTotalFrames, const uint32 InSampleRateHz, const float InGridSizePixels = 1.f, const FSlateFontInfo* InTicksTimeFont = nullptr);
 
 	void UpdateDisplayRange(const TRange<uint32> InDisplayRange);
 	bool UpdateGridMetrics(const float InGridSizePixels);
-	virtual const FSampledSequenceGridMetrics GetGridMetrics() const override;
+	virtual const FFixedSampledSequenceGridMetrics GetGridMetrics() const override;
 	void SetTicksTimeFont(const FSlateFontInfo* InNewFont);
 	const float SnapPositionToClosestFrame(const float InPixelPosition) const;
 
 	FOnWaveformEditorGridUpdated OnGridMetricsUpdated;
 
 private:
-	FSampledSequenceGridMetrics GridMetrics;
+	FFixedSampledSequenceGridMetrics GridMetrics;
 	uint32 TotalFrames = 0;
 	TRange<uint32> DisplayRange;
 
diff --git a/Engine/Plugins/Runtime/AudioWidgets/Source/AudioWidgets/Private/SFixedSampledSequenceViewer.cpp b/Engine/Plugins/Runtime/AudioWidgets/Source/AudioWidgets/Private/SFixedSampledSequenceViewer.cpp
new file mode 100644
--- /dev/null
+++ b/Engine/Plugins/Runtime/AudioWidgets/Source/AudioWidgets/Private/SFixedSampledSequenceViewer.cpp
@@ -0,0 +1,263 @@
+// Copyright Epic Games, Inc. All Rights Reserved.
+
+#include "SFixedSampledSequenceViewer.h"
+
+void SFixedSampledSequenceViewer::Construct(const FArguments& InArgs, TArrayView<const float> InSampleData, const uint8 InNumChannels, TSharedRef<IFixedSampledSequenceGridService> InGridService)
+{
+	GridService = InGridService;
+	UpdateView(InSampleData, InNumChannels);
+
+	Style = InArgs._Style;
+	DrawingParams = InArgs._SequenceDrawingParams;
+
+	check(Style);
+	SequenceColor = Style->SequenceColor;
+	MajorGridLineColor = Style->MajorGridLineColor;
+	MinorGridLineColor = Style->MinorGridLineColor;
+	BackgroundColor = Style->SequenceBackgroundColor;
+	BackgroundBrush = Style->BackgroundBrush;
+	DesiredWidth = Style->DesiredWidth;
+	DesiredHeight = Style->DesiredHeight;
+	SampleMarkersSize = Style->SampleMarkersSize;
+	SequenceLineThickness = Style->SequenceLineThickness;
+	ZeroCrossingLineColor = Style->ZeroCrossingLineColor;
+	ZeroCrossingLineThickness = Style->ZeroCrossingLineThickness;
+}
+
+void SFixedSampledSequenceViewer::UpdateView(TArrayView<const float> InSampleData, const uint8 InNumChannels)
+{
+	UpdateGridMetrics();
+
+	SampleData = InSampleData;
+	NumChannels = InNumChannels;
+
+	bForceRedraw = true;
+}
+
+int32 SFixedSampledSequenceViewer::OnPaint(const FPaintArgs& Args, const FGeometry& AllottedGeometry, const FSlateRect& MyCullingRect, FSlateWindowElementList& OutDrawElements, int32 LayerId, const FWidgetStyle& InWidgetStyle, bool bParentEnabled) const 
+{
+	float PixelWidth = MyCullingRect.GetSize().X;
+
+	FSlateDrawElement::MakeBox(
+		OutDrawElements,
+		++LayerId,
+		AllottedGeometry.ToPaintGeometry(),
+		&BackgroundBrush,
+		ESlateDrawEffect::None, 
+		BackgroundColor.GetSpecifiedColor()
+	);
+
+	if (PixelWidth > 0)
+	{
+		DrawGridLines(AllottedGeometry, OutDrawElements, LayerId);
+
+		if (SequenceDrawMode == ESequenceDrawMode::BinnedPeaks)
+		{
+			TArray<FVector2D> BinDrawPoints;
+			BinDrawPoints.SetNumUninitialized(2);
+
+			for (const SampledSequenceDrawingUtils::FSampleBinCoordinates& BinCoordinates : CachedBinsDrawCoordinates)
+			{
+				BinDrawPoints[0] = BinCoordinates.Top;
+				BinDrawPoints[1] = BinCoordinates.Bottom;
+
+				FSlateDrawElement::MakeLines(
+					OutDrawElements,
+					LayerId,
+					AllottedGeometry.ToPaintGeometry(),
+					BinDrawPoints,
+					ESlateDrawEffect::None,
+					SequenceColor.GetSpecifiedColor()
+				);
+			}
+
+		}
+		else if (SequenceDrawMode == ESequenceDrawMode::SequenceLine)
+		{
+			TArray<FVector2D> SamplesDrawCoordinates;
+			const uint16 NChannels = NumChannels;
+
+			for (uint16 Channel = 0; Channel < NChannels; ++Channel)
+			{
+				for (int32 SampleIndex = Channel; SampleIndex < CachedSampleDrawCoordinates.Num(); SampleIndex += NChannels)
+				{
+					SamplesDrawCoordinates.Add(CachedSampleDrawCoordinates[SampleIndex]);
+					
+					const FVector2D SampleBoxSize = FVector2D(SampleMarkersSize, SampleMarkersSize);
+					const FVector2D SampleBoxTransform = FVector2D(CachedSampleDrawCoordinates[SampleIndex] - (SampleBoxSize / 2.f));
+
+					const FPaintGeometry SampleBoxGeometry = AllottedGeometry.ToPaintGeometry(
+						SampleBoxSize,
+						FSlateLayoutTransform(SampleBoxTransform)
+					);
+
+					FSlateDrawElement::MakeBox(
+						OutDrawElements,
+						LayerId,
+						SampleBoxGeometry,
+						&BackgroundBrush,
+						ESlateDrawEffect::None,
+						SequenceColor.GetSpecifiedColor()
+					);
+				}
+
+				FSlateDrawElement::MakeLines(
+					OutDrawElements,
+					LayerId,
+					AllottedGeometry.ToPaintGeometry(),
+					SamplesDrawCoordinates,
+					ESlateDrawEffect::None,
+					SequenceColor.GetSpecifiedColor(),
+					true,
+					SequenceLineThickness
+				);
+
+				SamplesDrawCoordinates.Empty();
+			}
+		}
+
+		++LayerId;
+	}
+
+	return LayerId;
+}
+
+FVector2D SFixedSampledSequenceViewer::ComputeDesiredSize(float) const
+{
+	return FVector2D(DesiredWidth, DesiredHeight);
+}
+
+void SFixedSampledSequenceViewer::Tick(const FGeometry& AllottedGeometry, const double InCurrentTime, const float InDeltaTime)
+{
+	uint32 DiscretePixelWidth = FMath::FloorToInt(AllottedGeometry.GetLocalSize().X);
+	if (DiscretePixelWidth <= 0)
+	{
+		return;
+	}
+
+	if (DiscretePixelWidth != CachedPixelWidth || bForceRedraw)
+	{
+		CachedPixelWidth = DiscretePixelWidth;
+		CachedPixelHeight = AllottedGeometry.GetLocalSize().Y;
+		const float NumSamplesToDisplay = SampleData.Num() / (float) NumChannels;
+		SequenceDrawMode = NumSamplesToDisplay > CachedPixelWidth ? ESequenceDrawMode::BinnedPeaks : ESequenceDrawMode::SequenceLine;
+
+		if (SequenceDrawMode == ESequenceDrawMode::BinnedPeaks)
+		{
+			SampledSequenceDrawingUtils::GroupInterleavedSampledTSIntoMinMaxBins<float>(CachedPeaks, CachedPixelWidth, SampleData.GetData(), SampleData.Num(), GridMetrics.SampleRate, NumChannels);
+			SampledSequenceDrawingUtils::GenerateSampleBinsCoordinatesForGeometry(CachedBinsDrawCoordinates, AllottedGeometry, CachedPeaks, NumChannels, DrawingParams);
+		}
+		else
+		{
+			SampledSequenceDrawingUtils::GenerateSequencedSamplesCoordinatesForGeometry(CachedSampleDrawCoordinates, SampleData, AllottedGeometry, NumChannels, GridMetrics, DrawingParams);
+		}
+		
+		bForceRedraw = false;
+	}
+	else if (CachedPixelHeight != AllottedGeometry.GetLocalSize().Y)
+	{
+		CachedPixelHeight = AllottedGeometry.GetLocalSize().Y;
+
+		if (SequenceDrawMode == ESequenceDrawMode::BinnedPeaks)
+		{
+			SampledSequenceDrawingUtils::GenerateSampleBinsCoordinatesForGeometry(CachedBinsDrawCoordinates, AllottedGeometry, CachedPeaks, NumChannels, DrawingParams);
+		}
+		else
+		{
+			SampledSequenceDrawingUtils::GenerateSequencedSamplesCoordinatesForGeometry(CachedSampleDrawCoordinates, SampleData, AllottedGeometry, NumChannels, GridMetrics, DrawingParams);
+		}
+	}
+}
+
+void SFixedSampledSequenceViewer::DrawGridLines(const FGeometry& AllottedGeometry, FSlateWindowElementList& OutDrawElements, int32& LayerId) const
+{
+	TArray<FVector2D> LinePoints;
+	LinePoints.SetNumUninitialized(2);
+
+	const double MinorGridXStep = GridMetrics.MajorGridXStep / GridMetrics.NumMinorGridDivisions;
+
+	for (double CurrentMajorLineX = GridMetrics.FirstMajorTickX; CurrentMajorLineX < AllottedGeometry.Size.X; CurrentMajorLineX += GridMetrics.MajorGridXStep)
+	{
+		const double MajorLineX = CurrentMajorLineX;
+
+		LinePoints[0] = FVector2D(MajorLineX, 0.f);
+		LinePoints[1] = FVector2D(MajorLineX, AllottedGeometry.Size.Y);
+
+		FSlateDrawElement::MakeLines(
+			OutDrawElements,
+			++LayerId,
+			AllottedGeometry.ToPaintGeometry(),
+			LinePoints,
+			ESlateDrawEffect::None,
+			MajorGridLineColor.GetSpecifiedColor(),
+			false);
+
+
+		for (int32 MinorLineIndex = 1; MinorLineIndex < GridMetrics.NumMinorGridDivisions; ++MinorLineIndex)
+		{
+			const double MinorLineX = MajorLineX + MinorGridXStep * MinorLineIndex;
+
+			LinePoints[0] = FVector2D(MinorLineX, 0.);
+			LinePoints[1] = FVector2D(MinorLineX, AllottedGeometry.Size.Y);
+
+			FSlateDrawElement::MakeLines(
+				OutDrawElements,
+				++LayerId,
+				AllottedGeometry.ToPaintGeometry(),
+				LinePoints,
+				ESlateDrawEffect::None,
+				MinorGridLineColor.GetSpecifiedColor(),
+				false);
+
+		}
+	}
+
+	const uint16 NChannels = NumChannels;
+
+	for (uint16 Channel = 0; Channel < NChannels; ++Channel)
+	{
+		const SampledSequenceDrawingUtils::FHorizontalDimensionSlot ChannelBoundaries(Channel, NChannels, AllottedGeometry);
+
+		LinePoints[0] = FVector2D(0.f, ChannelBoundaries.Center);
+		LinePoints[1] = FVector2D(AllottedGeometry.Size.X, ChannelBoundaries.Center);
+
+		FSlateDrawElement::MakeLines(
+			OutDrawElements,
+			++LayerId,
+			AllottedGeometry.ToPaintGeometry(),
+			LinePoints,
+			ESlateDrawEffect::None,
+			ZeroCrossingLineColor.GetSpecifiedColor(),
+			false,
+			ZeroCrossingLineThickness);
+
+	}
+}
+
+void SFixedSampledSequenceViewer::UpdateGridMetrics()
+{
+	check(GridService)
+	GridMetrics = GridService->GetGridMetrics();
+}
+
+void SFixedSampledSequenceViewer::OnStyleUpdated(const FNotifyingAudioWidgetStyle& UpdatedStyle)
+{
+	check(Style);
+
+	if (&UpdatedStyle != Style)
+	{
+		return;
+	}
+
+	SequenceColor = Style->SequenceColor;
+	MajorGridLineColor = Style->MajorGridLineColor;
+	MinorGridLineColor = Style->MinorGridLineColor;
+	ZeroCrossingLineColor = Style->ZeroCrossingLineColor;
+	ZeroCrossingLineThickness = Style->ZeroCrossingLineThickness;
+	BackgroundColor = Style->SequenceBackgroundColor;
+	BackgroundBrush = Style->BackgroundBrush;
+	SampleMarkersSize = Style->SampleMarkersSize;
+	SequenceLineThickness = Style->SequenceLineThickness;
+	DesiredWidth = Style->DesiredWidth;
+	DesiredHeight = Style->DesiredHeight;
+}
diff --git a/Engine/Plugins/Runtime/AudioWidgets/Source/AudioWidgets/Private/SSampledSequenceViewer.cpp b/Engine/Plugins/Runtime/AudioWidgets/Source/AudioWidgets/Private/SSampledSequenceViewer.cpp
deleted file mode 100644
diff --git a/Engine/Plugins/Runtime/AudioWidgets/Source/AudioWidgets/Private/SampledSequenceDrawingUtils.cpp b/Engine/Plugins/Runtime/AudioWidgets/Source/AudioWidgets/Private/SampledSequenceDrawingUtils.cpp
new file mode 100644
--- /dev/null
+++ b/Engine/Plugins/Runtime/AudioWidgets/Source/AudioWidgets/Private/SampledSequenceDrawingUtils.cpp
@@ -0,0 +1,100 @@
+// Copyright Epic Games, Inc. All Rights Reserved.
+
+#include "SampledSequenceDrawingUtils.h"
+
+#include "Layout/Geometry.h"
+
+SampledSequenceDrawingUtils::FHorizontalDimensionSlot::FHorizontalDimensionSlot(const uint16 DimensionToDraw, const uint16 TotalNumDimensions, const FGeometry& InAllottedGeometry)
+{
+	Height = InAllottedGeometry.GetLocalSize().Y / TotalNumDimensions;
+	float ChannelSlotMidPoint = Height / 2.f;
+
+	Top = Height * DimensionToDraw;
+	Center = Top + ChannelSlotMidPoint;
+	Bottom = Top + Height;
+}
+
+void SampledSequenceDrawingUtils::GenerateSampleBinsCoordinatesForGeometry(TArray<FSampleBinCoordinates>& OutDrawCoordinates, const FGeometry& InAllottedGeometry, const TArray<TRange<float>>& InSampleBins, const uint16 NDimensions, const FSampledSequenceDrawingParams Params)
+{
+	if (!ensure(NDimensions != 0))
+	{
+		return;
+	}
+
+	const uint32 PixelWidth = FMath::FloorToInt(InAllottedGeometry.GetLocalSize().X);
+	check(PixelWidth * NDimensions == InSampleBins.Num());
+
+	const float HeightScale = InAllottedGeometry.GetLocalSize().Y / (2.f * Params.MaxDisplayedValue * NDimensions) * Params.MaxSequenceHeight * Params.VerticalZoomFactor;
+
+
+	OutDrawCoordinates.SetNumUninitialized(InSampleBins.Num());
+	FSampleBinCoordinates* OutCoordinatesData = OutDrawCoordinates.GetData();
+
+	for (uint16 Channel = 0; Channel < NDimensions; ++Channel)
+	{
+		const FHorizontalDimensionSlot ChannelBoundaries(Channel, NDimensions, InAllottedGeometry);
+
+		for (uint32 Pixel = 0; Pixel < PixelWidth; ++Pixel)
+		{
+			uint32 PeakIndex = Pixel * NDimensions + Channel;
+			const TRange<float>& MinMaxBin = InSampleBins[PeakIndex];
+
+			const float SampleMaxScaled = MinMaxBin.GetUpperBoundValue() * HeightScale > Params.MinScaledBinValue ? MinMaxBin.GetUpperBoundValue() * HeightScale : Params.MinSequenceHeight;
+			const float SampleMinScaled = MinMaxBin.GetLowerBoundValue() * HeightScale < -1 * Params.MinScaledBinValue ? MinMaxBin.GetLowerBoundValue() * HeightScale : -1.f * Params.MinSequenceHeight;
+
+			const float Top = FMath::Max(ChannelBoundaries.Center - SampleMaxScaled, ChannelBoundaries.Top + Params.DimensionSlotMargin);
+			const float Bottom = FMath::Min(ChannelBoundaries.Center - SampleMinScaled, ChannelBoundaries.Bottom - Params.DimensionSlotMargin);
+
+			OutCoordinatesData[PeakIndex].Top = FVector2D(Pixel, Top);
+			OutCoordinatesData[PeakIndex].Bottom = FVector2D(Pixel, Bottom);
+		}
+	}
+}
+
+void SampledSequenceDrawingUtils::GenerateSequencedSamplesCoordinatesForGeometry(TArray<FVector2D>& OutDrawCoordinates, TArrayView<const float> InSampleData, const FGeometry& InAllottedGeometry, const uint16 NDimensions, const FFixedSampledSequenceGridMetrics InGridMetrics, const FSampledSequenceDrawingParams Params)
+{
+	OutDrawCoordinates.Empty();
+
+	if (!ensure(NDimensions != 0))
+	{
+		return;
+	}
+
+	const uint32 NumFramesToDisplay = InSampleData.Num() / NDimensions;
+
+	TArray<FHorizontalDimensionSlot> ChannelSlotsBoundaries;
+
+	for (uint16 Channel = 0; Channel < NDimensions; ++Channel)
+	{
+		ChannelSlotsBoundaries.Emplace(Channel, NDimensions, InAllottedGeometry);
+	}
+
+	uint32 FrameIndex = 0;
+	const float XDrawMargin = InAllottedGeometry.Size.X + InGridMetrics.PixelsPerFrame;
+
+	for (double FrameX = InGridMetrics.FirstMajorTickX; FrameX < XDrawMargin; FrameX += InGridMetrics.PixelsPerFrame)
+	{
+		const double XRatio = FrameX / InAllottedGeometry.Size.X;
+
+		if (FrameX >= 0 && FrameIndex < NumFramesToDisplay)
+		{
+			for (uint16 Channel = 0; Channel < NDimensions; ++Channel)
+			{
+				const FHorizontalDimensionSlot& ChannelBoundaries = ChannelSlotsBoundaries[Channel];
+				const uint32 SampleIndex = FrameIndex * NDimensions + Channel;
+				const float SampleValue = InSampleData[SampleIndex];
+
+				const float SampleValueRatio = SampleValue / Params.MaxDisplayedValue * Params.MaxSequenceHeight * Params.VerticalZoomFactor;
+				const float TopBoundary = ChannelBoundaries.Top + Params.DimensionSlotMargin;
+				const float BottomBoundary = ChannelBoundaries.Bottom - Params.DimensionSlotMargin;
+				const float SampleY = FMath::Clamp((SampleValueRatio * ChannelBoundaries.Height / 2.f) + ChannelBoundaries.Center, TopBoundary, BottomBoundary);
+
+				const FVector2D SampleCoordinates(FrameX, SampleY);
+
+				OutDrawCoordinates.Add(SampleCoordinates);
+			}
+
+			FrameIndex++;
+		}
+	}
+}
diff --git a/Engine/Plugins/Runtime/AudioWidgets/Source/AudioWidgets/Private/SamplesSequenceTransportCoordinator.cpp b/Engine/Plugins/Runtime/AudioWidgets/Source/AudioWidgets/Private/SamplesSequenceTransportCoordinator.cpp
deleted file mode 100644
diff --git a/Engine/Plugins/Runtime/AudioWidgets/Source/AudioWidgets/Private/SparseSampledSequenceTransportCoordinator.cpp b/Engine/Plugins/Runtime/AudioWidgets/Source/AudioWidgets/Private/SparseSampledSequenceTransportCoordinator.cpp
new file mode 100644
--- /dev/null
+++ b/Engine/Plugins/Runtime/AudioWidgets/Source/AudioWidgets/Private/SparseSampledSequenceTransportCoordinator.cpp
@@ -0,0 +1,230 @@
+// Copyright Epic Games, Inc. All Rights Reserved.
+
+#include "SparseSampledSequenceTransportCoordinator.h"
+
+/** Zoom Formula
+
+E.g: Coordinator is receiving a playback ratio (PBR) in the range 0 to 1
+To position the focus point accordingly, this would the formula to follow:
+
+FP = (PBR - LB) / ZR
+
+FP = Focus Point
+ZR = Zoom Ratio
+LB = Lower Displayed Waveform Bound
+from here, you can rotate as needed. 
+e.g to find PBR
+PBR = (FP * ZR) + LB
+
+
+
+This formula can be used to position any element (e.g. a marker) given in a 0 to 1 position range where
+0 is the beginning time of the waveform and 1 the end time
+
+Zoom is defined as a ratio = displayed length / total length
+
+*/
+
+void FSparseSampledSequenceTransportCoordinator::UpdatePlaybackRange(const TRange<float>& NewRange)
+{
+	ProgressRange = NewRange;
+}
+
+const bool FSparseSampledSequenceTransportCoordinator::IsScrubbing() const
+{
+	return bIsScrubbing;
+}
+
+float FSparseSampledSequenceTransportCoordinator::ConvertAbsoluteRatioToZoomed(const float InAbsoluteRatio) const
+{
+	return (InAbsoluteRatio - DisplayRange.GetLowerBoundValue()) / ZoomRatio;
+}
+
+float FSparseSampledSequenceTransportCoordinator::ConvertZoomedRatioToAbsolute(const float InZoomedRatio) const
+{
+	return InZoomedRatio * ZoomRatio + DisplayRange.GetLowerBoundValue();
+}
+
+const TRange<float> FSparseSampledSequenceTransportCoordinator::GetDisplayRange() const
+{
+	return DisplayRange;
+}
+
+void FSparseSampledSequenceTransportCoordinator::MoveFocusPoint(const float InFocusPoint)
+{
+	FocusPoint = InFocusPoint;
+}
+
+void FSparseSampledSequenceTransportCoordinator::UpdateZoomRatioAndDisplayRange(const float NewZoomRatio)
+{
+	const float ClampedZoomRatio = FMath::Clamp(NewZoomRatio, UE_KINDA_SMALL_NUMBER, 1.f);
+	const float ZoomRatioDelta = ZoomRatio - ClampedZoomRatio;
+
+ 	const float DeltaOutsideBoundaries = DisplayRange.GetUpperBoundValue() - (DisplayRange.GetUpperBoundValue() - ZoomRatioDelta);
+
+	const float DeltaStepL = DeltaOutsideBoundaries * FocusPoint;
+	const float DeltaStepR = DeltaOutsideBoundaries * (1.f - FocusPoint);
+
+	float MinDisplayRangeValue = DeltaStepL + DisplayRange.GetLowerBoundValue();
+	float MaxDisplayRangeValue = DisplayRange.GetUpperBoundValue() - DeltaStepR;
+
+	const float CurrentPlaybackPosition = (ZoomRatio * FocusPoint) + DisplayRange.GetLowerBoundValue();
+
+	//if Min < 0 or Max > 1 we can't zoom out further on that side
+
+	//hence we shift the zooming entirely on the opposite side
+	if (MinDisplayRangeValue < 0.f)
+	{
+		const float ExceedingDelta = FMath::Abs(MinDisplayRangeValue);
+		MaxDisplayRangeValue = MaxDisplayRangeValue + ExceedingDelta;
+		MinDisplayRangeValue = 0.f;
+	}
+
+	else if (MaxDisplayRangeValue > 1.f)
+	{
+		const float ExceedingDelta = MaxDisplayRangeValue - 1.f;
+
+		MinDisplayRangeValue = MinDisplayRangeValue - ExceedingDelta;
+		MaxDisplayRangeValue = 1.f;
+	}
+
+	const float NewPlayheadPosition = (CurrentPlaybackPosition - MinDisplayRangeValue) / ClampedZoomRatio;
+
+	ZoomRatio = ClampedZoomRatio;
+	UpdateDisplayRange(MinDisplayRangeValue, MaxDisplayRangeValue);
+	MoveFocusPoint(NewPlayheadPosition);
+}
+
+void FSparseSampledSequenceTransportCoordinator::Stop()
+{
+	SetProgressRatio(0.f);
+}
+
+const float FSparseSampledSequenceTransportCoordinator::GetFocusPoint() const
+{
+	return FocusPoint;
+}
+
+void FSparseSampledSequenceTransportCoordinator::ScrubFocusPoint(const float InTargetFocusPoint, const bool bIsMoving)
+{
+	const float LowerPlaybackBoundPosition = (ProgressRange.GetLowerBoundValue() - DisplayRange.GetLowerBoundValue()) / ZoomRatio;
+	const float UpperPlaybackBoundPosition = (ProgressRange.GetUpperBoundValue() - DisplayRange.GetLowerBoundValue()) / ZoomRatio;
+
+	float ClampedTargetPosition = FMath::Clamp(InTargetFocusPoint, LowerPlaybackBoundPosition, UpperPlaybackBoundPosition);
+	float NewPosition = FMath::Clamp(ClampedTargetPosition, 0.f, 1.f);
+
+	if (ClampedTargetPosition > FocusPointLockPosition)
+	{
+		const float BoundsDelta = (ClampedTargetPosition - FocusPointLockPosition) * ZoomRatio;
+		const float NewDisplayUpBound = FMath::Clamp(DisplayRange.GetUpperBoundValue() + BoundsDelta, 0.f, 1.f);
+
+		if (NewDisplayUpBound < 1.f)
+		{
+			const float NewDisplayLowBound = FMath::Clamp(DisplayRange.GetLowerBoundValue() + BoundsDelta, 0.f, 1.f);
+			UpdateDisplayRange(NewDisplayLowBound, NewDisplayUpBound);
+			NewPosition = FocusPointLockPosition;
+		}
+	}
+	if (ClampedTargetPosition < 0.f)
+	{
+		const float BoundsDelta = ClampedTargetPosition * ZoomRatio;
+		const float NewDisplayLowBound = FMath::Clamp(DisplayRange.GetLowerBoundValue() + BoundsDelta, 0.f, 1.f);
+
+		if (NewDisplayLowBound > 0.f)
+		{
+
+			const float NewDisplayUpBound = FMath::Clamp(DisplayRange.GetUpperBoundValue() + BoundsDelta, 0.f, 1.f);
+
+			UpdateDisplayRange(NewDisplayLowBound, NewDisplayUpBound);
+
+			NewPosition = 0.f;
+		}
+	}
+
+	MoveFocusPoint(NewPosition);
+
+	if (OnFocusPointScrubUpdate.IsBound())
+	{
+		OnFocusPointScrubUpdate.Broadcast(GetPlayBackRatioFromFocusPoint(NewPosition), bIsMoving);
+	}
+
+	bIsScrubbing = bIsMoving;
+}
+
+
+
+void FSparseSampledSequenceTransportCoordinator::SetProgressRatio(const float NewRatio)
+{	 
+	if (bIsScrubbing)
+	{
+		return;
+	}
+
+	check(NewRatio >= 0.f && NewRatio <= 1.f);
+
+	const float PlaybackRangeScale = ProgressRange.Size<float>();
+	const float ScaledPlaybackRatio = NewRatio * PlaybackRangeScale + ProgressRange.GetLowerBoundValue();
+
+	const float NewPlayheadPosition = (ScaledPlaybackRatio - DisplayRange.GetLowerBoundValue()) / ZoomRatio;
+
+	MoveFocusPoint(NewPlayheadPosition);
+
+	if (FocusPoint < 0.f)
+	{
+		//if sound is looping and we are zoomed, playhead position will be negative when new loop starts
+		//hence we reset the view to go back at the beginning of the file
+
+		MoveFocusPoint(0.f);
+		UpdateDisplayRange(FocusPoint, ZoomRatio);
+	}
+
+	else if (FocusPoint >= FocusPointLockPosition)
+	{
+
+		const float NewDisplayUpBound = FMath::Clamp(ScaledPlaybackRatio + (ZoomRatio * (1.f - FocusPointLockPosition)), 0.f, 1.f);
+		float NewDisplayLowBound = NewDisplayUpBound < 1.f ? ScaledPlaybackRatio - (ZoomRatio * FocusPointLockPosition) : 1.f - ZoomRatio;
+
+		NewDisplayLowBound = FMath::Clamp(NewDisplayLowBound, 0.f, 1.f);
+
+		if (NewDisplayUpBound < 1.f)
+		{
+			//we only lock the playhead if the display upper bound is before the end of the waveform
+			//otherwise we let the playhead go to the end
+
+			MoveFocusPoint(FocusPointLockPosition);
+		}
+
+		UpdateDisplayRange(NewDisplayLowBound, NewDisplayUpBound);
+	}
+}
+
+void FSparseSampledSequenceTransportCoordinator::SetZoomRatio(const float NewLevel)
+{
+	UpdateZoomRatioAndDisplayRange(NewLevel);
+}
+
+void FSparseSampledSequenceTransportCoordinator::UpdateDisplayRange(const float MinValue, const float MaxValue)
+{
+	check(MinValue < MaxValue);
+
+	DisplayRange.SetLowerBoundValue(MinValue);
+	DisplayRange.SetUpperBoundValue(MaxValue);
+
+	if (OnDisplayRangeUpdated.IsBound())
+	{
+		OnDisplayRangeUpdated.Broadcast(DisplayRange);
+	}
+}
+
+bool FSparseSampledSequenceTransportCoordinator::IsRatioWithinDisplayRange(const float Ratio) const
+{
+	return DisplayRange.Contains(Ratio);
+}
+
+float FSparseSampledSequenceTransportCoordinator::GetPlayBackRatioFromFocusPoint(const float InFocusPoint) const
+{
+	check(InFocusPoint >= 0.f && InFocusPoint <= 1.f);
+
+	const float PlayBackRatio = ((InFocusPoint * ZoomRatio) + DisplayRange.GetLowerBoundValue() - ProgressRange.GetLowerBoundValue()) / ProgressRange.Size<float>();
+	return PlayBackRatio;
+}
\ No newline at end of file
diff --git a/Engine/Plugins/Runtime/AudioWidgets/Source/AudioWidgets/Private/TimeSeriesDrawingUtils.cpp b/Engine/Plugins/Runtime/AudioWidgets/Source/AudioWidgets/Private/TimeSeriesDrawingUtils.cpp
deleted file mode 100644
diff --git a/Engine/Plugins/Runtime/AudioWidgets/Source/AudioWidgets/Public/IFixedSampledSequenceGridService.h b/Engine/Plugins/Runtime/AudioWidgets/Source/AudioWidgets/Public/IFixedSampledSequenceGridService.h
new file mode 100644
--- /dev/null
+++ b/Engine/Plugins/Runtime/AudioWidgets/Source/AudioWidgets/Public/IFixedSampledSequenceGridService.h
@@ -0,0 +1,23 @@
+// Copyright Epic Games, Inc. All Rights Reserved.
+
+#pragma once
+
+#include "Delegates/Delegate.h"
+
+struct FFixedSampledSequenceGridMetrics
+{
+	int32 NumMinorGridDivisions = 0;
+	uint32 SampleRate = 0;
+	uint32 StartFrame = 0;
+	double PixelsPerFrame = 0;
+	double FirstMajorTickX = 0;
+	double MajorGridXStep = 0;
+};
+
+class IFixedSampledSequenceGridService
+{
+public:
+	virtual ~IFixedSampledSequenceGridService() = default;
+	virtual const FFixedSampledSequenceGridMetrics GetGridMetrics() const = 0;
+	virtual const float SnapPositionToClosestFrame(const float InPixelPosition) const = 0;
+};
\ No newline at end of file
diff --git a/Engine/Plugins/Runtime/AudioWidgets/Source/AudioWidgets/Public/ISampledSequenceGridService.h b/Engine/Plugins/Runtime/AudioWidgets/Source/AudioWidgets/Public/ISampledSequenceGridService.h
deleted file mode 100644
diff --git a/Engine/Plugins/Runtime/AudioWidgets/Source/AudioWidgets/Public/ISamplesSequenceTransportCoordinator.h b/Engine/Plugins/Runtime/AudioWidgets/Source/AudioWidgets/Public/ISamplesSequenceTransportCoordinator.h
deleted file mode 100644
diff --git a/Engine/Plugins/Runtime/AudioWidgets/Source/AudioWidgets/Public/ISparseSampledSequenceTransportCoordinator.h b/Engine/Plugins/Runtime/AudioWidgets/Source/AudioWidgets/Public/ISparseSampledSequenceTransportCoordinator.h
new file mode 100644
--- /dev/null
+++ b/Engine/Plugins/Runtime/AudioWidgets/Source/AudioWidgets/Public/ISparseSampledSequenceTransportCoordinator.h
@@ -0,0 +1,23 @@
+// Copyright Epic Games, Inc. All Rights Reserved.
+
+#pragma once
+
+#include "Math/Range.h"
+
+class AUDIOWIDGETS_API ISparseSampledSequenceTransportCoordinator
+{
+public:
+	virtual ~ISparseSampledSequenceTransportCoordinator() = default;
+
+	virtual const TRange<float> GetDisplayRange() const = 0;
+	virtual const float GetFocusPoint() const = 0;
+	virtual void ScrubFocusPoint(const float InTargetFocusPoint, const bool bIsMoving) = 0;
+	virtual const bool IsScrubbing() const = 0;
+
+	virtual void SetProgressRatio(const float NewRatio) = 0;
+	virtual void SetZoomRatio(const float NewRatio) = 0;
+
+	virtual float ConvertAbsoluteRatioToZoomed(const float InAbsoluteRatio) const = 0;
+	virtual float ConvertZoomedRatioToAbsolute(const float InZoomedRatio) const = 0;
+
+};
\ No newline at end of file
diff --git a/Engine/Plugins/Runtime/AudioWidgets/Source/AudioWidgets/Public/SFixedSampledSequenceViewer.h b/Engine/Plugins/Runtime/AudioWidgets/Source/AudioWidgets/Public/SFixedSampledSequenceViewer.h
new file mode 100644
--- /dev/null
+++ b/Engine/Plugins/Runtime/AudioWidgets/Source/AudioWidgets/Public/SFixedSampledSequenceViewer.h
@@ -0,0 +1,79 @@
+// Copyright Epic Games, Inc. All Rights Reserved.
+
+#pragma once
+
+#include "AudioWidgetsSlateTypes.h"
+#include "IFixedSampledSequenceGridService.h"
+#include "Styling/ISlateStyle.h"
+#include "Styling/SlateWidgetStyleAsset.h"
+#include "SampledSequenceDrawingUtils.h"
+#include "Widgets/SLeafWidget.h"
+
+namespace SampledSequenceViewerWidget
+{
+	typedef TRange<float> SampleRange;
+}
+
+class AUDIOWIDGETS_API SFixedSampledSequenceViewer : public SLeafWidget
+{
+public:
+
+	SLATE_BEGIN_ARGS(SFixedSampledSequenceViewer) 
+	{
+	}
+	
+	SLATE_ARGUMENT(SampledSequenceDrawingUtils::FSampledSequenceDrawingParams, SequenceDrawingParams)
+
+	SLATE_STYLE_ARGUMENT(FSampledSequenceViewerStyle, Style)
+		
+	SLATE_END_ARGS()
+
+	void Construct(const FArguments& InArgs, TArrayView<const float> InSampleData, const uint8 InNumChannels, TSharedRef<IFixedSampledSequenceGridService> InGridService);
+	void UpdateView(TArrayView<const float> InSampleData, const uint8 InNumChannels);
+	
+	void OnStyleUpdated(const FNotifyingAudioWidgetStyle& UpdatedStyle);
+
+	void UpdateGridMetrics();
+
+private:
+	virtual int32 OnPaint(const FPaintArgs& Args, const FGeometry& AllottedGeometry, const FSlateRect& MyCullingRect, FSlateWindowElementList& OutDrawElements, int32 LayerId, const FWidgetStyle& InWidgetStyle, bool bParentEnabled) const override;
+	virtual FVector2D ComputeDesiredSize(float) const override;
+	virtual void Tick(const FGeometry& AllottedGeometry, const double InCurrentTime, const float InDeltaTime) override;
+
+	void DrawGridLines(const FGeometry& AllottedGeometry, FSlateWindowElementList& OutDrawElements, int32& LayerId) const;
+
+	bool bForceRedraw = false;
+
+	TArrayView<const float> SampleData;
+	uint8 NumChannels = 0;
+	float DurationInSeconds = 0.f;
+
+	TSharedPtr<IFixedSampledSequenceGridService> GridService = nullptr;
+	FFixedSampledSequenceGridMetrics GridMetrics;
+
+	const FSampledSequenceViewerStyle* Style = nullptr;
+	FSlateBrush BackgroundBrush;
+	FSlateColor BackgroundColor = FLinearColor(0.02f, 0.02f, 0.02f, 1.f);
+	FSlateColor SequenceColor = FLinearColor::White;
+	FSlateColor MajorGridLineColor = FLinearColor::Black;
+	FSlateColor MinorGridLineColor = FLinearColor(0.f, 0.f, 0.f, 0.5f);
+	FSlateColor ZeroCrossingLineColor = FLinearColor::Black;
+	float ZeroCrossingLineThickness = 1.f;
+	float DesiredHeight = 0.f;
+	float DesiredWidth = 0.f;
+	float SampleMarkersSize = 2.5f;
+	float SequenceLineThickness = 1.f;
+	
+
+	uint32 CachedPixelWidth = 0; 
+	float CachedPixelHeight = 0.f;
+	TArray<SampledSequenceViewerWidget::SampleRange> CachedPeaks;
+	TArray<SampledSequenceDrawingUtils::FSampleBinCoordinates> CachedBinsDrawCoordinates;
+	SampledSequenceDrawingUtils::FSampledSequenceDrawingParams DrawingParams;
+	TArray<FVector2D> CachedSampleDrawCoordinates;
+
+	enum class ESequenceDrawMode {
+		BinnedPeaks = 0,
+		SequenceLine
+	} SequenceDrawMode;
+};
\ No newline at end of file
diff --git a/Engine/Plugins/Runtime/AudioWidgets/Source/AudioWidgets/Public/SSampledSequenceViewer.h b/Engine/Plugins/Runtime/AudioWidgets/Source/AudioWidgets/Public/SSampledSequenceViewer.h
deleted file mode 100644
diff --git a/Engine/Plugins/Runtime/AudioWidgets/Source/AudioWidgets/Public/SampledSequenceDrawingUtils.h b/Engine/Plugins/Runtime/AudioWidgets/Source/AudioWidgets/Public/SampledSequenceDrawingUtils.h
new file mode 100644
--- /dev/null
+++ b/Engine/Plugins/Runtime/AudioWidgets/Source/AudioWidgets/Public/SampledSequenceDrawingUtils.h
@@ -0,0 +1,157 @@
+// Copyright Epic Games, Inc. All Rights Reserved.
+
+#pragma once
+
+#include "IFixedSampledSequenceGridService.h"
+#include "Math/Vector2D.h"
+#include "Math/Range.h"
+
+struct FGeometry;
+
+namespace SampledSequenceDrawingUtils
+{
+	/**
+	* Groups samples evenly into a number of desired bins.
+	* Each bin contains the min and max values of the grouped samples.
+	*
+	* @param OutBins			TArray where the bins will be written to
+	* @param NumDesiredBins		Number of output bins
+	* @param RawDataPtr			Ptr to the beginning of the samples
+	* @param TotalNumSamples	The total number of samples of the input time series
+	* @param NDimensions		Number of interleaved dimensions in the time series
+	* @param StartRatio			The ratio of the total number of frames (in a range of 0-1) at which grouping should start
+	* @param EndRatio			The ratio of the total number of frames (in a range of 0-1) at which grouping should end.
+	*
+	*/
+	template<typename SamplesType>
+	void GroupInterleavedSamplesIntoMinMaxBins(TArray<TRange<SamplesType>>& OutBins, const uint32 NumDesiredBins, const SamplesType* RawDataPtr, const uint32 TotalNumSamples, const uint16 NDimensions = 1, const double StartRatio = 0.0, double EndRatio = 1.0)
+	{
+		check(StartRatio >= 0.f && StartRatio < EndRatio);
+		check(EndRatio <= 1.f);
+
+		uint32 NumPeaks = NumDesiredBins * NDimensions;
+		OutBins.SetNumUninitialized(NumPeaks);
+		const uint32 NumFrames = TotalNumSamples / NDimensions;
+
+		double FramesPerBin = ((EndRatio - StartRatio) * NumFrames) / NumDesiredBins;
+		double IterationStartFrame = StartRatio * NumFrames;
+		uint32 FirstBinnedFrame = FMath::RoundToInt(StartRatio * NumFrames);
+
+		for (uint32 Peak = 0; Peak < NumPeaks; Peak += NDimensions)
+		{
+			uint32 LastBinnedFrame = IterationStartFrame + FramesPerBin;
+
+			for (uint16 Channel = 0; Channel < NDimensions; ++Channel)
+			{
+				SamplesType MaxSampleValue = TNumericLimits<SamplesType>::Min();
+				SamplesType MinSampleValue = TNumericLimits<SamplesType>::Max();
+
+				for (uint32 Frame = FirstBinnedFrame; Frame < LastBinnedFrame; ++Frame)
+				{
+					const uint32 SampleIndex = Frame * NDimensions + Channel;
+					check(SampleIndex < TotalNumSamples);
+					SamplesType SampleValue = RawDataPtr[SampleIndex];
+
+					MinSampleValue = FMath::Min(MinSampleValue, SampleValue);
+					MaxSampleValue = FMath::Max(MaxSampleValue, SampleValue);
+
+				}
+
+				OutBins[Peak + Channel] = TRange<SamplesType>(MinSampleValue, MaxSampleValue);
+			}
+
+			IterationStartFrame += FramesPerBin;
+			FirstBinnedFrame = LastBinnedFrame;
+		}
+	}
+
+	/**
+	* Groups samples of a time series into an equal number of desired bins.
+	* Each bin contains the min and max values of the grouped samples.
+	*
+	* @param OutBins			TArray where the bins will be written to
+	* @param NumDesiredBins		Number of output bins
+	* @param RawDataPtr			Ptr to the beginning of the samples
+	* @param TotalNumSamples	The total number of samples of the input time series
+	* @param SampleRate			SampleRate of the time series
+	* @param NDimensions		Number of interleaved dimensions in the time series
+	* @param StartTime			The time at which grouping should start
+	* @param EndTime 			The time at which grouping should end.
+	*
+	* Note: With a negative EndTime, the method will calculate automatically the EndTime by doing
+	* TotalNumSamples / (SampleRate * NDimensions)
+	*
+	*/
+	template<typename SamplesType>
+	void GroupInterleavedSampledTSIntoMinMaxBins(TArray<TRange<SamplesType>>& OutBins, const uint32 NumDesiredBins, const SamplesType* RawDataPtr, const uint32 TotalNumSamples, const uint32 SampleRate, const uint16 NDimensions = 1, const float StartTime = 0.f, float EndTime = -1.f)
+	{
+		const double TotalTime = TotalNumSamples / ((float)SampleRate * NDimensions);
+		const double StartRatio = StartTime / TotalTime;
+		const double EndRatio = EndTime >= 0.f ? FMath::Clamp(EndTime / TotalTime, StartRatio, 1.0) : 1.0;
+
+		GroupInterleavedSamplesIntoMinMaxBins(OutBins, NumDesiredBins, RawDataPtr, TotalNumSamples, NDimensions, StartRatio, EndRatio);
+	}
+
+	/**
+		* Params for drawing sequences of samples/sample bins
+		* @param MaxDisplayedValue					The highest value a sample can take
+		* @param DimensionSlotMargin				Margin to keep from the dimension slot boundaries (pixels)
+		* @param MinSequenceHeight					The minimum height (ratio) the drawn sequence can take in a channel slot
+		* @param MaxSequenceHeight					The maximum height (ratio) the drawn sequence can take in a channel slot
+		* @param MinScaledBinValue					Minimum value a scaled bin can have
+		* @param VerticalZoomFactor					Sequence zoom factor
+	*/
+	struct FSampledSequenceDrawingParams
+	{
+		double MaxDisplayedValue = 1.f;
+		float DimensionSlotMargin = 2.f;
+		float MaxSequenceHeight = 0.9f;
+		float MinSequenceHeight = 0.1f;
+		float MinScaledBinValue = 0.001f;
+		float VerticalZoomFactor = 1.f;
+	};
+
+	/**
+		* Represents an horizontal slot in which samples or sample bins can be drawn
+	*/
+	struct FHorizontalDimensionSlot
+	{
+		explicit FHorizontalDimensionSlot(const uint16 DimensionToDraw, const uint16 TotalNumDimensions, const FGeometry& InAllottedGeometry);
+
+		float Top;
+		float Center;
+		float Bottom;
+		float Height;
+	};
+
+	struct FSampleBinCoordinates
+	{
+		FVector2D Top;
+		FVector2D Bottom;
+	};
+
+	/**
+		* Generates an array of coordinates to draw binned samples as vertical lines.
+		* Coordinates will be generated for an horizontal view
+		* 
+		* @param OutDrawCoordinates			The generated coordinates
+		* @param InAllottedGeometry			The geometry to draw into
+		* @param InSampleBins				The bins to draw
+		* @param NDimensions				Number of interleaved dimensions in the time series
+		* @param Params						Drawing Params
+	*/
+	void GenerateSampleBinsCoordinatesForGeometry(TArray<FSampleBinCoordinates>& OutDrawCoordinates, const FGeometry& InAllottedGeometry, const TArray<TRange<float>>& InSampleBins, const uint16 NDimensions, const FSampledSequenceDrawingParams Params = FSampledSequenceDrawingParams());
+
+	/**
+	* Generates an array of coordinates to draw single samples in succession horizontally.
+	*
+	* @param OutDrawCoordinates			The generated coordinates
+	* @param InSampleData				The samples to draw
+	* @param InAllottedGeometry			The geometry to draw into
+	* @param NDimensions				Number of interleaved dimensions in the time series
+	* @param InGridMetrics				GridMetrics followed to lay down the samples
+	* @param Params						Drawing Params
+	*/
+	void GenerateSequencedSamplesCoordinatesForGeometry(TArray<FVector2D>& OutDrawCoordinates, TArrayView<const float> InSampleData, const FGeometry& InAllottedGeometry, const uint16 NDimensions, const FFixedSampledSequenceGridMetrics InGridMetrics, const FSampledSequenceDrawingParams Params = FSampledSequenceDrawingParams());
+
+}
\ No newline at end of file
diff --git a/Engine/Plugins/Runtime/AudioWidgets/Source/AudioWidgets/Public/SamplesSequenceTransportCoordinator.h b/Engine/Plugins/Runtime/AudioWidgets/Source/AudioWidgets/Public/SamplesSequenceTransportCoordinator.h
deleted file mode 100644
diff --git a/Engine/Plugins/Runtime/AudioWidgets/Source/AudioWidgets/Public/SparseSampledSequenceTransportCoordinator.h b/Engine/Plugins/Runtime/AudioWidgets/Source/AudioWidgets/Public/SparseSampledSequenceTransportCoordinator.h
new file mode 100644
--- /dev/null
+++ b/Engine/Plugins/Runtime/AudioWidgets/Source/AudioWidgets/Public/SparseSampledSequenceTransportCoordinator.h
@@ -0,0 +1,58 @@
+// Copyright Epic Games, Inc. All Rights Reserved.
+
+#pragma once
+
+#include "Delegates/Delegate.h"
+#include "ISparseSampledSequenceTransportCoordinator.h"
+
+DECLARE_MULTICAST_DELEGATE_OneParam(FOnDisplayRangeUpdated, const TRange<float> /* New Display Range */);
+DECLARE_MULTICAST_DELEGATE_TwoParams(FOnFocusPointScrubUpdate, const float /* Focused Playback Ratio */, const bool /*Playhead is Moving*/);
+
+
+class AUDIOWIDGETS_API FSparseSampledSequenceTransportCoordinator : public ISparseSampledSequenceTransportCoordinator
+{
+public:
+	FSparseSampledSequenceTransportCoordinator() = default;
+	virtual ~FSparseSampledSequenceTransportCoordinator() = default;
+
+	/** Called when the playhead is scrubbed */
+	FOnFocusPointScrubUpdate OnFocusPointScrubUpdate;
+
+	/** Called when the display range is updated */
+	FOnDisplayRangeUpdated OnDisplayRangeUpdated;
+
+	/** ISparseSampledSequenceTransportCoordinator interface */
+	const TRange<float> GetDisplayRange() const;
+	const float GetFocusPoint() const override;
+	void ScrubFocusPoint(const float InTargetFocusPoint, const bool bIsMoving) override;
+	const bool IsScrubbing() const;
+	void SetProgressRatio(const float NewRatio) override;
+	void SetZoomRatio(const float NewRatio) override;
+	float ConvertAbsoluteRatioToZoomed(const float InAbsoluteRatio) const override;
+	float ConvertZoomedRatioToAbsolute(const float InZoomedRatio) const override;
+
+	void UpdatePlaybackRange(const TRange<float>& NewRange);
+	void Stop();
+
+private:
+	FORCEINLINE void MoveFocusPoint(const float InFocusPoint);
+	void UpdateZoomRatioAndDisplayRange(const float NewZoomRatio);
+
+	void UpdateDisplayRange(const float MinValue, const float MaxValue);
+	bool IsRatioWithinDisplayRange(const float Ratio) const;
+
+	float GetPlayBackRatioFromFocusPoint(const float InFocusPoint) const;
+	
+	float CurrentPlaybackRatio = 0.f;
+	float FocusPointLockPosition = 0.95f;
+	float FocusPoint = 0.f;
+	float ZoomRatio = 1.f;
+	
+	/* The currently displayed render data range */
+	TRange<float> DisplayRange = TRange<float>::Inclusive(0.f, 1.f);
+
+	/* Progress range to scale the incoming progress ratio with*/
+	TRange<float> ProgressRange = TRange<float>::Inclusive(0.f, 1.f);
+	
+	bool bIsScrubbing = false;
+};
\ No newline at end of file
diff --git a/Engine/Plugins/Runtime/AudioWidgets/Source/AudioWidgets/Public/TimeSeriesDrawingUtils.h b/Engine/Plugins/Runtime/AudioWidgets/Source/AudioWidgets/Public/TimeSeriesDrawingUtils.h
deleted file mode 100644
