[WaveformEditor] Transport coordinator uses double precision to handle display and playback range

diff --git a/Engine/Plugins/Editor/WaveformEditor/Source/WaveformEditor/Private/WaveformEditor.cpp b/Engine/Plugins/Editor/WaveformEditor/Source/WaveformEditor/Private/WaveformEditor.cpp
--- a/Engine/Plugins/Editor/WaveformEditor/Source/WaveformEditor/Private/WaveformEditor.cpp
+++ b/Engine/Plugins/Editor/WaveformEditor/Source/WaveformEditor/Private/WaveformEditor.cpp
@@ -686,13 +686,13 @@ void FWaveformEditor::HandleRenderDataUpdate()
 	if (TransportCoordinator != nullptr)
 	{
 		TransportCoordinator->UpdatePlaybackRange(WaveformView.DataProvider->GetTransformedWaveformBounds());
-		WaveformView.DataProvider->RequestSequenceView(TRange<double>::Inclusive(TransportCoordinator->GetDisplayRange().GetLowerBoundValue(), TransportCoordinator->GetDisplayRange().GetUpperBoundValue()));
+		WaveformView.DataProvider->RequestSequenceView(TransportCoordinator->GetDisplayRange());
 	}
 }
 
-void FWaveformEditor::HandleDisplayRangeUpdate(const TRange<float> NewRange)
+void FWaveformEditor::HandleDisplayRangeUpdate(const TRange<double> NewRange)
 {
-	WaveformView.DataProvider->RequestSequenceView(TRange<double>::Inclusive(NewRange.GetLowerBoundValue(), NewRange.GetUpperBoundValue()));
+	WaveformView.DataProvider->RequestSequenceView(NewRange);
 }
 
 void FWaveformEditor::HandlePlayheadScrub(const float InTargetPlayBackRatio, const bool bIsMoving)
diff --git a/Engine/Plugins/Editor/WaveformEditor/Source/WaveformEditor/Public/WaveformEditor.h b/Engine/Plugins/Editor/WaveformEditor/Source/WaveformEditor/Public/WaveformEditor.h
--- a/Engine/Plugins/Editor/WaveformEditor/Source/WaveformEditor/Public/WaveformEditor.h
+++ b/Engine/Plugins/Editor/WaveformEditor/Source/WaveformEditor/Public/WaveformEditor.h
@@ -75,7 +75,7 @@ private:
 
 	/* Delegates Handlers*/
 	void HandleRenderDataUpdate();
-	void HandleDisplayRangeUpdate(const TRange<float>);
+	void HandleDisplayRangeUpdate(const TRange<double>);
 
 	/** FGCObject overrides */
 	virtual void AddReferencedObjects(FReferenceCollector& Collector) override;
diff --git a/Engine/Plugins/Editor/WaveformEditor/Source/WaveformEditorWidgets/Private/SWaveformTransformationsOverlay.cpp b/Engine/Plugins/Editor/WaveformEditor/Source/WaveformEditorWidgets/Private/SWaveformTransformationsOverlay.cpp
--- a/Engine/Plugins/Editor/WaveformEditor/Source/WaveformEditorWidgets/Private/SWaveformTransformationsOverlay.cpp
+++ b/Engine/Plugins/Editor/WaveformEditor/Source/WaveformEditorWidgets/Private/SWaveformTransformationsOverlay.cpp
@@ -91,7 +91,7 @@ void SWaveformTransformationsOverlay::UpdateLayerConstraints()
 	UpdateAnchors();
 }
 
-void SWaveformTransformationsOverlay::OnNewWaveformDisplayRange(const TRange<float> NewDisplayRange)
+void SWaveformTransformationsOverlay::OnNewWaveformDisplayRange(const TRange<double> NewDisplayRange)
 {
 	UpdateAnchors();
 }
diff --git a/Engine/Plugins/Editor/WaveformEditor/Source/WaveformEditorWidgets/Private/WaveformEditorSequenceDataProvider.cpp b/Engine/Plugins/Editor/WaveformEditor/Source/WaveformEditorWidgets/Private/WaveformEditorSequenceDataProvider.cpp
--- a/Engine/Plugins/Editor/WaveformEditor/Source/WaveformEditorWidgets/Private/WaveformEditorSequenceDataProvider.cpp
+++ b/Engine/Plugins/Editor/WaveformEditor/Source/WaveformEditorWidgets/Private/WaveformEditorSequenceDataProvider.cpp
@@ -195,10 +195,10 @@ void FWaveformEditorSequenceDataProvider::GenerateSequenceDataInternal()
 	}
 }
 
-void FWaveformEditorSequenceDataProvider::UpdateTransformedWaveformBounds(const uint32 FirstEditedSample, const uint32 LastEditedSample, const float NumOriginalSamples)
+void FWaveformEditorSequenceDataProvider::UpdateTransformedWaveformBounds(const uint32 FirstEditedSample, const uint32 LastEditedSample, const uint32 NumOriginalSamples)
 {
-	const float FirstEditedSampleOffset = FirstEditedSample / NumOriginalSamples;
-	const float LastEditedSampleOffset = (NumOriginalSamples - LastEditedSample) / NumOriginalSamples;
+	const double FirstEditedSampleOffset = FirstEditedSample / (double) NumOriginalSamples;
+	const double LastEditedSampleOffset = (NumOriginalSamples - LastEditedSample) / (double) NumOriginalSamples;
 
 	TransformedWaveformBounds.SetLowerBoundValue(FirstEditedSampleOffset);
 	TransformedWaveformBounds.SetUpperBoundValue(1 - LastEditedSampleOffset);
@@ -233,7 +233,7 @@ TArrayView<const FTransformationRenderLayerInfo> FWaveformEditorSequenceDataProv
 	return MakeArrayView(RenderLayers.GetData(), RenderLayers.Num());
 }
 
-const TRange<float> FWaveformEditorSequenceDataProvider::GetTransformedWaveformBounds() const
+const TRange<double> FWaveformEditorSequenceDataProvider::GetTransformedWaveformBounds() const
 {
 	return TransformedWaveformBounds;
 }
@@ -247,4 +247,4 @@ void FWaveformEditorSequenceDataProvider::CreateDurationHighlightLayer()
 
 	DurationHiglightLayer->SetTransformationWaveInfo(FWaveformTransformationRenderInfo());
 	RenderLayers.Emplace(DurationHiglightLayer, FTransformationLayerConstraints(0.f,1.f));
-}
\ No newline at end of file
+}
diff --git a/Engine/Plugins/Editor/WaveformEditor/Source/WaveformEditorWidgets/Public/SWaveformTransformationsOverlay.h b/Engine/Plugins/Editor/WaveformEditor/Source/WaveformEditorWidgets/Public/SWaveformTransformationsOverlay.h
--- a/Engine/Plugins/Editor/WaveformEditor/Source/WaveformEditorWidgets/Public/SWaveformTransformationsOverlay.h
+++ b/Engine/Plugins/Editor/WaveformEditor/Source/WaveformEditorWidgets/Public/SWaveformTransformationsOverlay.h
@@ -23,7 +23,7 @@ public:
 	void Construct(const FArguments& InArgs, TArrayView< const FTransformationLayerRenderInfo> InTransformationRenderers, TSharedRef<FSparseSampledSequenceTransportCoordinator> InTransportCoordinator);
 	void OnLayerChainGenerated(FTransformationLayerRenderInfo* FirstLayerPtr, const int32 NLayers);
 	void UpdateLayerConstraints();
-	void OnNewWaveformDisplayRange(const TRange<float> NewDisplayRange);
+	void OnNewWaveformDisplayRange(const TRange<double> NewDisplayRange);
 
 	FReply OnMouseButtonDown(const FGeometry& MyGeometry, const FPointerEvent& MouseEvent) override;
 	FReply OnMouseButtonUp(const FGeometry& MyGeometry, const FPointerEvent& MouseEvent) override;
diff --git a/Engine/Plugins/Editor/WaveformEditor/Source/WaveformEditorWidgets/Public/WaveformEditorSequenceDataProvider.h b/Engine/Plugins/Editor/WaveformEditor/Source/WaveformEditorWidgets/Public/WaveformEditorSequenceDataProvider.h
--- a/Engine/Plugins/Editor/WaveformEditor/Source/WaveformEditorWidgets/Public/WaveformEditorSequenceDataProvider.h
+++ b/Engine/Plugins/Editor/WaveformEditor/Source/WaveformEditorWidgets/Public/WaveformEditorSequenceDataProvider.h
@@ -70,7 +70,7 @@ public:
 
 	TArrayView<const FTransformationRenderLayerInfo> GetTransformLayers() const;
 
-	const TRange<float> GetTransformedWaveformBounds() const;
+	const TRange<double> GetTransformedWaveformBounds() const;
 
 	/** Called when a new layer chain of transformations UI is created */
 	FOnLayersChainGenerated OnLayersChainGenerated;
@@ -84,7 +84,7 @@ public:
 private:
 	void GenerateSequenceDataInternal();
 
-	void UpdateTransformedWaveformBounds(const uint32 FirstEditedSample, const uint32 LastEditedSample, const float NumOriginalSamples);
+	void UpdateTransformedWaveformBounds(const uint32 FirstEditedSample, const uint32 LastEditedSample, const uint32 NumOriginalSamples);
 
 	void CreateDurationHighlightLayer();
 
@@ -111,5 +111,5 @@ private:
 	uint32 NumOriginalWaveformFrames = 0;
 
 	/* The bounds of the transformed waveform in relation to the original */
-	TRange<float> TransformedWaveformBounds = TRange<float>::Inclusive(0.f, 1.f);
-};
\ No newline at end of file
+	TRange<double> TransformedWaveformBounds = TRange<double>::Inclusive(0, 1);
+};
diff --git a/Engine/Plugins/Runtime/AudioWidgets/Source/AudioWidgets/Private/SparseSampledSequenceTransportCoordinator.cpp b/Engine/Plugins/Runtime/AudioWidgets/Source/AudioWidgets/Private/SparseSampledSequenceTransportCoordinator.cpp
--- a/Engine/Plugins/Runtime/AudioWidgets/Source/AudioWidgets/Private/SparseSampledSequenceTransportCoordinator.cpp
+++ b/Engine/Plugins/Runtime/AudioWidgets/Source/AudioWidgets/Private/SparseSampledSequenceTransportCoordinator.cpp
@@ -25,7 +25,7 @@ Zoom is defined as a ratio = displayed length / total length
 
 */
 
-void FSparseSampledSequenceTransportCoordinator::UpdatePlaybackRange(const TRange<float>& NewRange)
+void FSparseSampledSequenceTransportCoordinator::UpdatePlaybackRange(const TRange<double>& NewRange)
 {
 	ProgressRange = NewRange;
 }
@@ -45,12 +45,12 @@ float FSparseSampledSequenceTransportCoordinator::ConvertZoomedRatioToAbsolute(c
 	return InZoomedRatio * ZoomRatio + DisplayRange.GetLowerBoundValue();
 }
 
-const TRange<float> FSparseSampledSequenceTransportCoordinator::GetDisplayRange() const
+const TRange<double> FSparseSampledSequenceTransportCoordinator::GetDisplayRange() const
 {
 	return DisplayRange;
 }
 
-void FSparseSampledSequenceTransportCoordinator::MoveFocusPoint(const float InFocusPoint)
+void FSparseSampledSequenceTransportCoordinator::MoveFocusPoint(const double InFocusPoint)
 {
 	FocusPoint = InFocusPoint;
 	OnFocusPointMoved.Broadcast(FocusPoint);
@@ -61,35 +61,35 @@ void FSparseSampledSequenceTransportCoordinator::UpdateZoomRatioAndDisplayRange(
 	const float ClampedZoomRatio = FMath::Clamp(NewZoomRatio, UE_KINDA_SMALL_NUMBER, 1.f);
 	const float ZoomRatioDelta = ZoomRatio - ClampedZoomRatio;
 
- 	const float DeltaOutsideBoundaries = DisplayRange.GetUpperBoundValue() - (DisplayRange.GetUpperBoundValue() - ZoomRatioDelta);
+ 	const double DeltaOutsideBoundaries = DisplayRange.GetUpperBoundValue() - (DisplayRange.GetUpperBoundValue() - ZoomRatioDelta);
 
-	const float DeltaStepL = DeltaOutsideBoundaries * FocusPoint;
-	const float DeltaStepR = DeltaOutsideBoundaries * (1.f - FocusPoint);
+	const double DeltaStepL = DeltaOutsideBoundaries * FocusPoint;
+	const double DeltaStepR = DeltaOutsideBoundaries * (1.f - FocusPoint);
 
-	float MinDisplayRangeValue = DeltaStepL + DisplayRange.GetLowerBoundValue();
-	float MaxDisplayRangeValue = DisplayRange.GetUpperBoundValue() - DeltaStepR;
+	double MinDisplayRangeValue = DeltaStepL + DisplayRange.GetLowerBoundValue();
+	double MaxDisplayRangeValue = DisplayRange.GetUpperBoundValue() - DeltaStepR;
 
-	const float CurrentPlaybackPosition = (ZoomRatio * FocusPoint) + DisplayRange.GetLowerBoundValue();
+	const double CurrentPlaybackPosition = (ZoomRatio * FocusPoint) + DisplayRange.GetLowerBoundValue();
 
 	//if Min < 0 or Max > 1 we can't zoom out further on that side
 
 	//hence we shift the zooming entirely on the opposite side
 	if (MinDisplayRangeValue < 0.f)
 	{
-		const float ExceedingDelta = FMath::Abs(MinDisplayRangeValue);
+		const double ExceedingDelta = FMath::Abs(MinDisplayRangeValue);
 		MaxDisplayRangeValue = MaxDisplayRangeValue + ExceedingDelta;
 		MinDisplayRangeValue = 0.f;
 	}
 
 	else if (MaxDisplayRangeValue > 1.f)
 	{
-		const float ExceedingDelta = MaxDisplayRangeValue - 1.f;
+		const double ExceedingDelta = MaxDisplayRangeValue - 1.f;
 
 		MinDisplayRangeValue = MinDisplayRangeValue - ExceedingDelta;
 		MaxDisplayRangeValue = 1.f;
 	}
 
-	const float NewPlayheadPosition = (CurrentPlaybackPosition - MinDisplayRangeValue) / ClampedZoomRatio;
+	const double NewPlayheadPosition = (CurrentPlaybackPosition - MinDisplayRangeValue) / ClampedZoomRatio;
 
 	ZoomRatio = ClampedZoomRatio;
 	UpdateDisplayRange(MinDisplayRangeValue, MaxDisplayRangeValue);
@@ -163,10 +163,10 @@ void FSparseSampledSequenceTransportCoordinator::SetProgressRatio(const float Ne
 
 	check(NewRatio >= 0.f && NewRatio <= 1.f);
 
-	const float PlaybackRangeScale = ProgressRange.Size<float>();
-	const float ScaledPlaybackRatio = NewRatio * PlaybackRangeScale + ProgressRange.GetLowerBoundValue();
+	const double PlaybackRangeScale = ProgressRange.Size<double>();
+	const double ScaledPlaybackRatio = NewRatio * PlaybackRangeScale + ProgressRange.GetLowerBoundValue();
 
-	const float NewPlayheadPosition = (ScaledPlaybackRatio - DisplayRange.GetLowerBoundValue()) / ZoomRatio;
+	const double NewPlayheadPosition = (ScaledPlaybackRatio - DisplayRange.GetLowerBoundValue()) / ZoomRatio;
 
 	MoveFocusPoint(NewPlayheadPosition);
 
@@ -182,10 +182,10 @@ void FSparseSampledSequenceTransportCoordinator::SetProgressRatio(const float Ne
 	else if (FocusPoint >= FocusPointLockPosition)
 	{
 
-		const float NewDisplayUpBound = FMath::Clamp(ScaledPlaybackRatio + (ZoomRatio * (1.f - FocusPointLockPosition)), 0.f, 1.f);
-		float NewDisplayLowBound = NewDisplayUpBound < 1.f ? ScaledPlaybackRatio - (ZoomRatio * FocusPointLockPosition) : 1.f - ZoomRatio;
+		const double NewDisplayUpBound = FMath::Clamp(ScaledPlaybackRatio + (ZoomRatio * (1.f - FocusPointLockPosition)), 0.f, 1.f);
+		double NewDisplayLowBound = NewDisplayUpBound < 1.f ? ScaledPlaybackRatio - (ZoomRatio * FocusPointLockPosition) : 1.f - ZoomRatio;
 
-		NewDisplayLowBound = FMath::Clamp(NewDisplayLowBound, 0.f, 1.f);
+		NewDisplayLowBound = FMath::Clamp(NewDisplayLowBound, 0, 1);
 
 		if (NewDisplayUpBound < 1.f)
 		{
@@ -204,7 +204,7 @@ void FSparseSampledSequenceTransportCoordinator::SetZoomRatio(const float NewLev
 	UpdateZoomRatioAndDisplayRange(NewLevel);
 }
 
-void FSparseSampledSequenceTransportCoordinator::UpdateDisplayRange(const float MinValue, const float MaxValue)
+void FSparseSampledSequenceTransportCoordinator::UpdateDisplayRange(const double MinValue, const double MaxValue)
 {
 	check(MinValue < MaxValue);
 
@@ -217,15 +217,15 @@ void FSparseSampledSequenceTransportCoordinator::UpdateDisplayRange(const float
 	}
 }
 
-bool FSparseSampledSequenceTransportCoordinator::IsRatioWithinDisplayRange(const float Ratio) const
+bool FSparseSampledSequenceTransportCoordinator::IsRatioWithinDisplayRange(const double Ratio) const
 {
 	return DisplayRange.Contains(Ratio);
 }
 
-float FSparseSampledSequenceTransportCoordinator::GetPlayBackRatioFromFocusPoint(const float InFocusPoint) const
+double FSparseSampledSequenceTransportCoordinator::GetPlayBackRatioFromFocusPoint(const double InFocusPoint) const
 {
 	check(InFocusPoint >= 0.f && InFocusPoint <= 1.f);
 
-	const float PlayBackRatio = ((InFocusPoint * ZoomRatio) + DisplayRange.GetLowerBoundValue() - ProgressRange.GetLowerBoundValue()) / ProgressRange.Size<float>();
+	const double PlayBackRatio = ((InFocusPoint * ZoomRatio) + DisplayRange.GetLowerBoundValue() - ProgressRange.GetLowerBoundValue()) / ProgressRange.Size<float>();
 	return PlayBackRatio;
 }
\ No newline at end of file
diff --git a/Engine/Plugins/Runtime/AudioWidgets/Source/AudioWidgets/Public/ISparseSampledSequenceTransportCoordinator.h b/Engine/Plugins/Runtime/AudioWidgets/Source/AudioWidgets/Public/ISparseSampledSequenceTransportCoordinator.h
--- a/Engine/Plugins/Runtime/AudioWidgets/Source/AudioWidgets/Public/ISparseSampledSequenceTransportCoordinator.h
+++ b/Engine/Plugins/Runtime/AudioWidgets/Source/AudioWidgets/Public/ISparseSampledSequenceTransportCoordinator.h
@@ -9,7 +9,7 @@ class AUDIOWIDGETS_API ISparseSampledSequenceTransportCoordinator
 public:
 	virtual ~ISparseSampledSequenceTransportCoordinator() = default;
 
-	virtual const TRange<float> GetDisplayRange() const = 0;
+	virtual const TRange<double> GetDisplayRange() const = 0;
 	virtual const float GetFocusPoint() const = 0;
 	virtual void ScrubFocusPoint(const float InTargetFocusPoint, const bool bIsMoving) = 0;
 	virtual const bool IsScrubbing() const = 0;
diff --git a/Engine/Plugins/Runtime/AudioWidgets/Source/AudioWidgets/Public/SparseSampledSequenceTransportCoordinator.h b/Engine/Plugins/Runtime/AudioWidgets/Source/AudioWidgets/Public/SparseSampledSequenceTransportCoordinator.h
--- a/Engine/Plugins/Runtime/AudioWidgets/Source/AudioWidgets/Public/SparseSampledSequenceTransportCoordinator.h
+++ b/Engine/Plugins/Runtime/AudioWidgets/Source/AudioWidgets/Public/SparseSampledSequenceTransportCoordinator.h
@@ -5,7 +5,7 @@
 #include "Delegates/Delegate.h"
 #include "ISparseSampledSequenceTransportCoordinator.h"
 
-DECLARE_MULTICAST_DELEGATE_OneParam(FOnDisplayRangeUpdated, const TRange<float> /* New Display Range */);
+DECLARE_MULTICAST_DELEGATE_OneParam(FOnDisplayRangeUpdated, const TRange<double> /* New Display Range */);
 DECLARE_MULTICAST_DELEGATE_OneParam(FOnFocusPointMoved, const float /* New Position */);
 DECLARE_MULTICAST_DELEGATE_TwoParams(FOnFocusPointScrubUpdate, const float /* Focused Playback Ratio */, const bool /*Playhead is Moving*/);
 
@@ -25,7 +25,7 @@ public:
 	FOnDisplayRangeUpdated OnDisplayRangeUpdated;
 
 	/** ISparseSampledSequenceTransportCoordinator interface */
-	const TRange<float> GetDisplayRange() const;
+	const TRange<double> GetDisplayRange() const;
 	const float GetFocusPoint() const override;
 	void ScrubFocusPoint(const float InTargetFocusPoint, const bool bIsMoving) override;
 	const bool IsScrubbing() const;
@@ -33,26 +33,27 @@ public:
 	void SetZoomRatio(const float NewRatio) override;
 	float ConvertAbsoluteRatioToZoomed(const float InAbsoluteRatio) const override;
 	float ConvertZoomedRatioToAbsolute(const float InZoomedRatio) const override;
-	void UpdatePlaybackRange(const TRange<float>& NewRange);
+	void UpdatePlaybackRange(const TRange<double>& NewRange);
 	void Stop();
 
 private:
-	FORCEINLINE void MoveFocusPoint(const float InFocusPoint);
+	FORCEINLINE void MoveFocusPoint(const double InFocusPoint);
 	void UpdateZoomRatioAndDisplayRange(const float NewZoomRatio);
-	void UpdateDisplayRange(const float MinValue, const float MaxValue);
-	bool IsRatioWithinDisplayRange(const float Ratio) const;
-	float GetPlayBackRatioFromFocusPoint(const float InFocusPoint) const;
 
-	float CurrentPlaybackRatio = 0.f;
-	float FocusPointLockPosition = 0.95f;
-	float FocusPoint = 0.f;
+	void UpdateDisplayRange(const double MinValue, const double MaxValue);
+	bool IsRatioWithinDisplayRange(const double Ratio) const;
+	double GetPlayBackRatioFromFocusPoint(const double InFocusPoint) const;
+
+	double CurrentPlaybackRatio = 0.f;
+	double FocusPointLockPosition = 0.95f;
+	double FocusPoint = 0;
 	float ZoomRatio = 1.f;
 
 	/* The currently displayed render data range */
-	TRange<float> DisplayRange = TRange<float>::Inclusive(0.f, 1.f);
+	TRange<double> DisplayRange = TRange<double>::Inclusive(0, 1);
 
 	/* Progress range to scale the incoming progress ratio with*/
-	TRange<float> ProgressRange = TRange<float>::Inclusive(0.f, 1.f);
+	TRange<double> ProgressRange = TRange<double>::Inclusive(0, 1);
 
 	bool bIsScrubbing = false;
 };
\ No newline at end of file
