[AudioWidgets] Create a widget to overlay a multichannel 2D grid and use it in wave editor to show decibel values of the inspected waveform

diff --git a/Engine/Plugins/Editor/WaveformEditor/Source/WaveformEditorWidgets/Private/STransformedWaveformViewPanel.cpp b/Engine/Plugins/Editor/WaveformEditor/Source/WaveformEditorWidgets/Private/STransformedWaveformViewPanel.cpp
--- a/Engine/Plugins/Editor/WaveformEditor/Source/WaveformEditorWidgets/Private/STransformedWaveformViewPanel.cpp
+++ b/Engine/Plugins/Editor/WaveformEditor/Source/WaveformEditorWidgets/Private/STransformedWaveformViewPanel.cpp
@@ -2,28 +2,37 @@
 
  #include "STransformedWaveformViewPanel.h"
 
+#include "AudioWidgetsSlateTypes.h"
+#include "DSP/Dsp.h"
 #include "SampledSequenceDisplayUnit.h"
-#include "SPlayheadOverlay.h"
 #include "SFixedSampledSequenceRuler.h"
 #include "SFixedSampledSequenceViewer.h"
-#include "SWaveformTransformationsOverlay.h"
+#include "SPlayheadOverlay.h"
+#include "SSampledSequenceValueGridOverlay.h"
 #include "SWaveformEditorInputRoutingOverlay.h"
+#include "SWaveformTransformationsOverlay.h"
 #include "WaveformEditorGridData.h"
 #include "WaveformEditorStyle.h"
+#include "WaveformEditorWidgetsSettings.h"
+#include "Widgets/Layout/SBorder.h"
 #include "Widgets/SBoxPanel.h"
 #include "Widgets/SCompoundWidget.h"
 #include "Widgets/SOverlay.h"
 
-void STransformedWaveformViewPanel::Construct(const FArguments& InArgs, const FFixedSampledSequenceView& InView)
+#define LOCTEXT_NAMESPACE "TransformedWaveformViewPanel"
+
+void STransformedWaveformViewPanel::Construct(const FArguments& InArgs, const FFixedSampledSequenceView& InData)
 {
 	DisplayUnit = ESampledSequenceDisplayUnit::Seconds;
-	DataView = InView;
+	DataView = InData;
 
 	WaveformEditorStyle = &FWaveformEditorStyle::Get();
 	check(WaveformEditorStyle);
 
 	SetUpGridData();
+	SetUpBackground();
 	SetUpWaveformViewer(GridData.ToSharedRef(), DataView);
+	SetUpValueGridOverlay();
 
 	if (InArgs._TransformationsOverlay)
 	{
@@ -34,6 +43,11 @@ void STransformedWaveformViewPanel::Construct(const FArguments& InArgs, const FF
 	SetUpInputRoutingOverlay();
 	SetUpTimeRuler(GridData.ToSharedRef());
 	SetUpInputOverrides(InArgs);
+
+	const UWaveformEditorWidgetsSettings* Settings = GetWaveformEditorWidgetsSettings();
+	check(Settings)
+	Settings->OnSettingChanged().AddSP(this, &STransformedWaveformViewPanel::OnWaveEditorWidgetSettingsUpdated);
+
 	CreateLayout();
 }
 
@@ -44,11 +58,24 @@ void STransformedWaveformViewPanel::CreateLayout()
 	check(InputRoutingOverlay);
 
 	TSharedPtr<SOverlay> WaveformView = SNew(SOverlay);
+
+	WaveformView->AddSlot()
+	[
+		BackgroundBorder.ToSharedRef()
+	];
+
+	WaveformView->AddSlot()
+	[
+		ValueGridOverlay.ToSharedRef()
+	];
+
 	WaveformView->AddSlot()
 	[
 		WaveformViewer.ToSharedRef()
 	];
 
+
+
 	if (WaveformTransformationsOverlay)
 	{
 		WaveformView->AddSlot()
@@ -119,12 +146,13 @@ void STransformedWaveformViewPanel::SetupPlayheadOverlay()
 	PlayheadOverlayStyle->OnStyleUpdated.AddSP(PlayheadOverlay.ToSharedRef(), &SPlayheadOverlay::OnStyleUpdated);
 }
 
-void STransformedWaveformViewPanel::SetUpWaveformViewer(TSharedRef<FWaveformEditorGridData> InGridData, const FFixedSampledSequenceView& InView)
+void STransformedWaveformViewPanel::SetUpWaveformViewer(TSharedRef<FWaveformEditorGridData> InGridData, const FFixedSampledSequenceView& InData)
 {
 	FSampledSequenceViewerStyle* WaveViewerStyle = &WaveformEditorStyle->GetRegisteredWidgetStyle<FSampledSequenceViewerStyle>("WaveformViewer.Style").Get();
 	check(WaveViewerStyle);
 
-	WaveformViewer = SNew(SFixedSampledSequenceViewer, InView.SampleData, InView.NumDimensions, InGridData).Style(WaveViewerStyle);
+	WaveformViewer = SNew(SFixedSampledSequenceViewer, InData.SampleData, InData.NumDimensions, InGridData).Style(WaveViewerStyle).HideBackground(true);
+		
 	WaveViewerStyle->OnStyleUpdated.AddSP(WaveformViewer.ToSharedRef(), &SFixedSampledSequenceViewer::OnStyleUpdated);
 	GridData->OnGridMetricsUpdated.AddSP(WaveformViewer.ToSharedRef(), &SFixedSampledSequenceViewer::UpdateGridMetrics);
 }
@@ -151,6 +179,11 @@ void STransformedWaveformViewPanel::ReceiveSequenceView(const FFixedSampledSeque
 	{
 		WaveformViewer->UpdateView(InView.SampleData, InView.NumDimensions);
 	}
+
+	if (ValueGridOverlay)
+	{
+		ValueGridOverlay->ForceRedraw();
+	}
 }
 
 void STransformedWaveformViewPanel::SetPlayheadRatio(const float InRatio)
@@ -230,6 +263,66 @@ void STransformedWaveformViewPanel::UpdatePlayheadPosition(const float PaintedWi
 	}
 }
 
+void STransformedWaveformViewPanel::UpdateBackground(const FNotifyingAudioWidgetStyle& UpdatedStyle)
+{
+	FSampledSequenceViewerStyle* ViewerStyle = &WaveformEditorStyle->GetRegisteredWidgetStyle<FSampledSequenceViewerStyle>("WaveformViewer.Style").Get();
+
+	check(ViewerStyle);
+
+	if (&UpdatedStyle != ViewerStyle)
+	{
+		return;
+	}
+
+	BackgroundBorder->SetBorderImage(&ViewerStyle->BackgroundBrush);
+	BackgroundBorder->SetBorderBackgroundColor(ViewerStyle->SequenceBackgroundColor);
+}
+
+void STransformedWaveformViewPanel::OnWaveEditorWidgetSettingsUpdated(const FName& PropertyName, const UWaveformEditorWidgetsSettings* Settings)
+{
+	if (PropertyName == GET_MEMBER_NAME_CHECKED(UWaveformEditorWidgetsSettings, ShowLoudnessGrid))
+	{
+		if (ValueGridOverlay)
+		{
+			ValueGridOverlay->SetHideGrid(!Settings->ShowLoudnessGrid);
+		}
+	}
+	else if (PropertyName == GET_MEMBER_NAME_CHECKED(UWaveformEditorWidgetsSettings, ShowLoudnessGridDecibelValues))
+	{
+		if (ValueGridOverlay)
+		{
+			ValueGridOverlay->SetHideLabels(!Settings->ShowLoudnessGridDecibelValues);
+		}
+	}
+	else if (PropertyName == GET_MEMBER_NAME_CHECKED(UWaveformEditorWidgetsSettings, MaxLoudnessGridDivisions))
+	{
+		if (ValueGridOverlay)
+		{
+			ValueGridOverlay->SetMaxDivisionParameter(Settings->MaxLoudnessGridDivisions);
+		}
+	}
+}
+
+const UWaveformEditorWidgetsSettings* STransformedWaveformViewPanel::GetWaveformEditorWidgetsSettings()
+{
+	const UWaveformEditorWidgetsSettings* WaveformEditorWidgetsSettings = GetDefault<UWaveformEditorWidgetsSettings>();
+	check(WaveformEditorWidgetsSettings);
+
+	return WaveformEditorWidgetsSettings;
+}
+
+void STransformedWaveformViewPanel::SetUpBackground()
+{
+	FSampledSequenceViewerStyle* ViewerStyle = &WaveformEditorStyle->GetRegisteredWidgetStyle<FSampledSequenceViewerStyle>("WaveformViewer.Style").Get();
+	check(ViewerStyle);
+
+	ViewerStyle->OnStyleUpdated.AddSP(this, &STransformedWaveformViewPanel::UpdateBackground);
+
+	BackgroundBorder = SNew(SBorder)
+		.BorderImage(&ViewerStyle->BackgroundBrush)
+		.BorderBackgroundColor(ViewerStyle->SequenceBackgroundColor);
+}
+
 void STransformedWaveformViewPanel::SetUpInputOverrides(const FArguments& InArgs)
 {
 	SetOnPlayheadOverlayMouseButtonUp(InArgs._OnPlayheadOverlayMouseButtonUp);
@@ -239,4 +332,40 @@ void STransformedWaveformViewPanel::SetUpInputOverrides(const FArguments& InArgs
 	SetOnTimeRulerMouseMove(InArgs._OnTimeRulerMouseMove);
 	
 	SetOnMouseWheel(InArgs._OnMouseWheel);
-}
\ No newline at end of file
+}
+
+void STransformedWaveformViewPanel::SetUpValueGridOverlay()
+{
+	auto ValueGridDBConverter = [](const double InLabelValue)
+	{
+		
+		const float AbsAmplitude = FMath::Abs(InLabelValue * 2.f - 1);
+		const float GridLineValue = Audio::ConvertToDecibels(AbsAmplitude);
+		const float MinDecibelNumericLimit = -160;
+		FNumberFormattingOptions Formatting; 
+		Formatting.SetMaximumFractionalDigits(1);
+		return FText::Format(LOCTEXT("WaveformEditorValueGridLabel", "{0}"), (GridLineValue <= MinDecibelNumericLimit ? LOCTEXT("WaveformEditorValueGridLabel", "-inf") : FText::AsNumber(GridLineValue, &Formatting)));
+	};
+
+	const UWaveformEditorWidgetsSettings* Settings = GetWaveformEditorWidgetsSettings();
+	const bool bShowGrid = Settings ? Settings->ShowLoudnessGrid : true;
+	const bool bShowLabels = Settings ? Settings->ShowLoudnessGridDecibelValues : true;
+	const uint32 NumValueGridDivision = Settings ? Settings->MaxLoudnessGridDivisions : 3;
+
+
+	FSampledSequenceValueGridOverlayStyle* ValueGridStyle = &WaveformEditorStyle->GetRegisteredWidgetStyle<FSampledSequenceValueGridOverlayStyle>("WaveformEditorValueGrid.Style").Get();
+	
+	ValueGridOverlay = SNew(SSampledSequenceValueGridOverlay)
+		.DivideMode(SampledSequenceValueGridOverlay::EGridDivideMode::MidSplit)
+		.ValueGridLabelGenerator(ValueGridDBConverter)
+		.NumDimensions(DataView.NumDimensions)
+		.MaxDivisionParameter(NumValueGridDivision)
+		.HideLabels(!bShowLabels)
+		.HideGrid(!bShowGrid)
+		.Style(ValueGridStyle);
+
+	ValueGridStyle->OnStyleUpdated.AddSP(ValueGridOverlay.ToSharedRef(), &SSampledSequenceValueGridOverlay::OnStyleUpdated);
+}
+
+#undef LOCTEXT_NAMESPACE
+
diff --git a/Engine/Plugins/Editor/WaveformEditor/Source/WaveformEditorWidgets/Private/WaveformEditorStyle.cpp b/Engine/Plugins/Editor/WaveformEditor/Source/WaveformEditorWidgets/Private/WaveformEditorStyle.cpp
--- a/Engine/Plugins/Editor/WaveformEditor/Source/WaveformEditorWidgets/Private/WaveformEditorStyle.cpp
+++ b/Engine/Plugins/Editor/WaveformEditor/Source/WaveformEditorWidgets/Private/WaveformEditorStyle.cpp
@@ -6,9 +6,11 @@
 #include "Styling/SlateStyleRegistry.h"
 #include "WaveformEditorWidgetsSettings.h"
 
-static FLazyName WaveformViewerStyleName("WaveformViewer.Style");
 static FLazyName PlayheadOverlayStyleName("WaveformEditorPlayheadOverlay.Style");
+static FLazyName ValueGridOverlayStyleName("WaveformEditorValueGrid.Style");
 static FLazyName WaveformEditorRulerStyleName("WaveformEditorRuler.Style");
+static FLazyName WaveformViewerStyleName("WaveformViewer.Style");
+
 
 FName FWaveformEditorStyle::StyleName("WaveformEditorStyle");
 TUniquePtr<FWaveformEditorStyle> FWaveformEditorStyle::StyleInstance = nullptr;
@@ -51,7 +53,7 @@ void FWaveformEditorStyle::Init()
 		.SetSampleMarkersSize(Settings->SampleMarkersSize)
 		.SetMajorGridLineColor(Settings->MajorGridColor)
 		.SetMinorGridLineColor(Settings->MinorGridColor)
-		.SetZeroCrossingLineColor(Settings->ZeroCrossingLineColor)
+		.SetZeroCrossingLineColor(Settings->LoudnessGridColor)
 		.SetZeroCrossingLineThickness(Settings->ZeroCrossingLineThickness);
 
 	StyleInstance->Set(WaveformViewerStyleName, WaveViewerStyle);
@@ -70,8 +72,16 @@ void FWaveformEditorStyle::Init()
 
 	StyleInstance->Set(WaveformEditorRulerStyleName, TimeRulerStyle);
 
-	FSlateStyleRegistry::RegisterSlateStyle(StyleInstance->Get());
+	//ValueGridOverlayStyle
+	FSampledSequenceValueGridOverlayStyle ValueGridOverlayStyle = FSampledSequenceValueGridOverlayStyle()
+		.SetGridColor(Settings->LoudnessGridColor)
+		.SetGridThickness(Settings->LoudnessGridThickness)
+		.SetLabelTextColor(Settings->LoudnessGridTextColor)
+		.SetLabelTextFontSize(Settings->LoudnessGridTextSize);
+
+	StyleInstance->Set(ValueGridOverlayStyleName, ValueGridOverlayStyle);
 
+	FSlateStyleRegistry::RegisterSlateStyle(StyleInstance->Get());
 }
 
 const UWaveformEditorWidgetsSettings* FWaveformEditorStyle::GetWidgetsSettings() 
@@ -87,6 +97,8 @@ void FWaveformEditorStyle::OnWidgetSettingsUpdated(const FName& PropertyName, co
 	TSharedRef<FSampledSequenceViewerStyle> WaveformViewerStyle = GetRegisteredWidgetStyle<FSampledSequenceViewerStyle>(WaveformViewerStyleName);
 	TSharedRef<FPlayheadOverlayStyle> PlayheadOverlayStyle = GetRegisteredWidgetStyle<FPlayheadOverlayStyle>(PlayheadOverlayStyleName);
 	TSharedRef<FFixedSampleSequenceRulerStyle> WaveformEditorTimeRulerStyle = GetRegisteredWidgetStyle<FFixedSampleSequenceRulerStyle>(WaveformEditorRulerStyleName);
+	TSharedRef<FSampledSequenceValueGridOverlayStyle> ValueGridOverlayStyle = GetRegisteredWidgetStyle<FSampledSequenceValueGridOverlayStyle>(ValueGridOverlayStyleName);
+
 
 	if (PropertyName == GET_MEMBER_NAME_CHECKED(UWaveformEditorWidgetsSettings, WaveformColor))
 	{
@@ -112,10 +124,6 @@ void FWaveformEditorStyle::OnWidgetSettingsUpdated(const FName& PropertyName, co
 	{
 		WaveformViewerStyle->SetMinorGridLineColor(Settings->MinorGridColor);
 	}
-	else if (PropertyName == GET_MEMBER_NAME_CHECKED(UWaveformEditorWidgetsSettings, ZeroCrossingLineColor))
-	{
-		WaveformViewerStyle->SetZeroCrossingLineColor(Settings->ZeroCrossingLineColor);
-	}
 	else if (PropertyName == GET_MEMBER_NAME_CHECKED(UWaveformEditorWidgetsSettings, ZeroCrossingLineThickness))
 	{
 		WaveformViewerStyle->SetZeroCrossingLineThickness(Settings->ZeroCrossingLineThickness);
@@ -141,4 +149,22 @@ void FWaveformEditorStyle::OnWidgetSettingsUpdated(const FName& PropertyName, co
 	{
 		WaveformEditorTimeRulerStyle->SetFontSize(Settings->RulerFontSize);
 	}
+	else if (PropertyName == GET_MEMBER_NAME_CHECKED(UWaveformEditorWidgetsSettings, LoudnessGridColor))
+	{
+		ValueGridOverlayStyle->SetGridColor(Settings->LoudnessGridColor);
+		WaveformViewerStyle->SetZeroCrossingLineColor(Settings->LoudnessGridColor);
+	}
+	else if (PropertyName == GET_MEMBER_NAME_CHECKED(UWaveformEditorWidgetsSettings, LoudnessGridThickness))
+	{
+		ValueGridOverlayStyle->SetGridThickness(Settings->LoudnessGridThickness);
+	}
+	else if (PropertyName == GET_MEMBER_NAME_CHECKED(UWaveformEditorWidgetsSettings, LoudnessGridTextColor))
+	{
+		ValueGridOverlayStyle->SetLabelTextColor(Settings->LoudnessGridTextColor);
+	}
+	else if (PropertyName == GET_MEMBER_NAME_CHECKED(UWaveformEditorWidgetsSettings, LoudnessGridTextSize))
+	{
+		ValueGridOverlayStyle->SetLabelTextFontSize(Settings->LoudnessGridTextSize);
+	}
+
 }
diff --git a/Engine/Plugins/Editor/WaveformEditor/Source/WaveformEditorWidgets/Private/WaveformEditorWidgetsSettings.cpp b/Engine/Plugins/Editor/WaveformEditor/Source/WaveformEditorWidgets/Private/WaveformEditorWidgetsSettings.cpp
--- a/Engine/Plugins/Editor/WaveformEditor/Source/WaveformEditorWidgets/Private/WaveformEditorWidgetsSettings.cpp
+++ b/Engine/Plugins/Editor/WaveformEditor/Source/WaveformEditorWidgets/Private/WaveformEditorWidgetsSettings.cpp
@@ -9,7 +9,6 @@ UWaveformEditorWidgetsSettings::UWaveformEditorWidgetsSettings(const FObjectInit
 	, WaveformLineThickness(1.f)
 	, SampleMarkersSize(2.5f)
 	, WaveformBackgroundColor(FLinearColor(0.02f, 0.02f, 0.02f, 1.f))
-	, ZeroCrossingLineColor(FLinearColor::Black)
 	, ZeroCrossingLineThickness(1.f)
 	, MajorGridColor(FLinearColor::Black)
 	, MinorGridColor(FLinearColor(0.f, 0.f, 0.f, 0.5f))
@@ -17,6 +16,14 @@ UWaveformEditorWidgetsSettings::UWaveformEditorWidgetsSettings(const FObjectInit
 	, RulerTicksColor(FLinearColor(1.f, 1.f, 1.f, 0.9f))
 	, RulerTextColor(FLinearColor(1.f, 1.f, 1.f, 0.9f))
 	, RulerFontSize(10.f)
+	, ShowLoudnessGrid(true)
+	, ShowLoudnessGridDecibelValues(true)
+	, MaxLoudnessGridDivisions(3)
+	, LoudnessGridThickness(1.f)
+	, LoudnessGridColor(FLinearColor::Black)
+	, LoudnessGridTextColor(FLinearColor(1.f, 1.f, 1.f, 0.9f))
+	, LoudnessGridTextSize(10.f)
+
 {
 }
 
diff --git a/Engine/Plugins/Editor/WaveformEditor/Source/WaveformEditorWidgets/Private/WaveformEditorWidgetsSettings.h b/Engine/Plugins/Editor/WaveformEditor/Source/WaveformEditorWidgets/Private/WaveformEditorWidgetsSettings.h
--- a/Engine/Plugins/Editor/WaveformEditor/Source/WaveformEditorWidgets/Private/WaveformEditorWidgetsSettings.h
+++ b/Engine/Plugins/Editor/WaveformEditor/Source/WaveformEditorWidgets/Private/WaveformEditorWidgetsSettings.h
@@ -47,9 +47,6 @@ public:
 	UPROPERTY(config, EditAnywhere, Category = "Waveform Viewer")
 	FLinearColor WaveformBackgroundColor;
 
-	UPROPERTY(config, EditAnywhere, Category = "Waveform Viewer")
-	FLinearColor ZeroCrossingLineColor;
-
 	UPROPERTY(config, EditAnywhere, Category = "Waveform Viewer", Meta = (ClampMin = "1", ClampMax = "10"))
 	float ZeroCrossingLineThickness;
 
@@ -71,6 +68,29 @@ public:
 	UPROPERTY(config, EditAnywhere, Category = "Ruler", Meta = (ClampMin = "1", ClampMax = "15"))
 	float RulerFontSize;
 
+	UPROPERTY(config, EditAnywhere, Category = "Loudness Grid")
+	bool ShowLoudnessGrid;
+
+	UPROPERTY(config, EditAnywhere, Category = "Loudness Grid")
+	bool ShowLoudnessGridDecibelValues;
+
+	UPROPERTY(config, EditAnywhere, Category = "Loudness Grid", Meta = (ClampMin = "1", ClampMax = "10"))
+	uint8 MaxLoudnessGridDivisions;
+
+	UPROPERTY(config, EditAnywhere, Category = "Loudness Grid", Meta = (ClampMin = "1", ClampMax = "10"))
+	float LoudnessGridThickness;
+
+	UPROPERTY(config, EditAnywhere, Category = "Loudness Grid")
+	FLinearColor LoudnessGridColor;
+
+	UPROPERTY(config, EditAnywhere, Category = "Loudness Grid")
+	FLinearColor LoudnessGridTextColor;
+
+	UPROPERTY(config, EditAnywhere, Category = "Loudness Grid", Meta = (ClampMin = "1", ClampMax = "15"))
+	float LoudnessGridTextSize;
+
+	
+
 private: 
 	static FOnWaveformEditorWidgetsSettingsChanged SettingsChangedDelegate;
 };
\ No newline at end of file
diff --git a/Engine/Plugins/Editor/WaveformEditor/Source/WaveformEditorWidgets/Public/STransformedWaveformViewPanel.h b/Engine/Plugins/Editor/WaveformEditor/Source/WaveformEditorWidgets/Public/STransformedWaveformViewPanel.h
--- a/Engine/Plugins/Editor/WaveformEditor/Source/WaveformEditorWidgets/Public/STransformedWaveformViewPanel.h
+++ b/Engine/Plugins/Editor/WaveformEditor/Source/WaveformEditorWidgets/Public/STransformedWaveformViewPanel.h
@@ -7,16 +7,21 @@
 #include "IFixedSampledSequenceViewReceiver.h"
 
 enum class ESampledSequenceDisplayUnit;
+struct FNotifyingAudioWidgetStyle;
+
+class FSparseSampledSequenceTransportCoordinator;
+class FWaveformEditorGridData;
+class FWaveformEditorRenderData;
+class FWaveformEditorStyle;
+class FWaveformEditorZoomController;
+class SBorder;
+class SFixedSampledSequenceRuler;
+class SFixedSampledSequenceViewer;
+class SPlayheadOverlay;
+class SSampledSequenceValueGridOverlay;
 class SWaveformEditorInputRoutingOverlay;
 class SWaveformTransformationsOverlay;
-class SPlayheadOverlay;
-class SFixedSampledSequenceViewer;
-class SFixedSampledSequenceRuler;
-class FWaveformEditorZoomController;
-class FWaveformEditorStyle;
-class FWaveformEditorRenderData;
-class FWaveformEditorGridData;
-class FSparseSampledSequenceTransportCoordinator;
+class UWaveformEditorWidgetsSettings;
 
 class WAVEFORMEDITORWIDGETS_API STransformedWaveformViewPanel : public SCompoundWidget, public IFixedSampledSequenceViewReceiver 
 {
@@ -50,7 +55,6 @@ public:
 	void SetOnTimeRulerMouseMove(FPointerEventHandler InEventHandler);
 	void SetOnMouseWheel(FPointerEventHandler InEventHandler);
 
-
 private:
 	void CreateLayout();
 
@@ -60,9 +64,15 @@ private:
 	void SetUpWaveformViewer(TSharedRef<FWaveformEditorGridData> InGridData, const FFixedSampledSequenceView& InView);
 	void SetUpInputRoutingOverlay();
 	void SetUpInputOverrides(const FArguments& InArgs);
+	void SetUpValueGridOverlay();
+	void SetUpBackground();
+
 
 	void UpdateDisplayUnit(const ESampledSequenceDisplayUnit InDisplayUnit);
 	void UpdatePlayheadPosition(const float PaintedWidth);
+	void UpdateBackground(const FNotifyingAudioWidgetStyle& UpdatedStyle);
+
+	void OnWaveEditorWidgetSettingsUpdated(const FName& PropertyName, const UWaveformEditorWidgetsSettings* Settings);
 
 	TSharedPtr<FWaveformEditorGridData> GridData;
 
@@ -71,6 +81,8 @@ private:
  	TSharedPtr<SWaveformTransformationsOverlay> WaveformTransformationsOverlay;
 	TSharedPtr<SWaveformEditorInputRoutingOverlay> InputRoutingOverlay;
 	TSharedPtr<SPlayheadOverlay> PlayheadOverlay;
+	TSharedPtr<SSampledSequenceValueGridOverlay> ValueGridOverlay;
+	TSharedPtr<SBorder> BackgroundBorder;
 
 	float CachedPixelWidth = 0.f;
 
@@ -81,4 +93,5 @@ private:
 	FFixedSampledSequenceView DataView;
 	
 	float CachedPlayheadRatio = 0.f;
+	const UWaveformEditorWidgetsSettings* GetWaveformEditorWidgetsSettings();
 };
\ No newline at end of file
diff --git a/Engine/Plugins/Runtime/AudioWidgets/Source/AudioWidgets/Private/AudioWidgetsSlateTypes.cpp b/Engine/Plugins/Runtime/AudioWidgets/Source/AudioWidgets/Private/AudioWidgetsSlateTypes.cpp
--- a/Engine/Plugins/Runtime/AudioWidgets/Source/AudioWidgets/Private/AudioWidgetsSlateTypes.cpp
+++ b/Engine/Plugins/Runtime/AudioWidgets/Source/AudioWidgets/Private/AudioWidgetsSlateTypes.cpp
@@ -154,4 +154,22 @@ void FFixedSampleSequenceRulerStyle::GetResources(TArray< const FSlateBrush* >&
 {
 	OutBrushes.Add(&HandleBrush);
 	OutBrushes.Add(&BackgroundBrush);
-}
\ No newline at end of file
+}
+
+FSampledSequenceValueGridOverlayStyle::FSampledSequenceValueGridOverlayStyle()
+	: GridColor(FLinearColor::Black)
+	, GridThickness(1.f)
+	, LabelTextColor(FLinearColor::Black)
+	, LabelTextFont(FAppStyle::GetFontStyle("Regular"))
+	, DesiredWidth(AudioWidgetStylesSharedParams::DefaultWidth)
+	, DesiredHeight(AudioWidgetStylesSharedParams::DefaultHeight)
+{
+}
+
+const FSampledSequenceValueGridOverlayStyle& FSampledSequenceValueGridOverlayStyle::GetDefault()
+{
+	static FSampledSequenceValueGridOverlayStyle Default;
+	return Default;
+}
+
+const FName FSampledSequenceValueGridOverlayStyle::TypeName("FSampledSequenceValueGridOverlayStyle");
diff --git a/Engine/Plugins/Runtime/AudioWidgets/Source/AudioWidgets/Private/SFixedSampledSequenceViewer.cpp b/Engine/Plugins/Runtime/AudioWidgets/Source/AudioWidgets/Private/SFixedSampledSequenceViewer.cpp
--- a/Engine/Plugins/Runtime/AudioWidgets/Source/AudioWidgets/Private/SFixedSampledSequenceViewer.cpp
+++ b/Engine/Plugins/Runtime/AudioWidgets/Source/AudioWidgets/Private/SFixedSampledSequenceViewer.cpp
@@ -9,6 +9,7 @@ void SFixedSampledSequenceViewer::Construct(const FArguments& InArgs, TArrayView
 
 	Style = InArgs._Style;
 	DrawingParams = InArgs._SequenceDrawingParams;
+	bHideBackground = InArgs._HideBackground;
 
 	check(Style);
 	SequenceColor = Style->SequenceColor;
@@ -38,14 +39,18 @@ int32 SFixedSampledSequenceViewer::OnPaint(const FPaintArgs& Args, const FGeomet
 {
 	float PixelWidth = MyCullingRect.GetSize().X;
 
-	FSlateDrawElement::MakeBox(
-		OutDrawElements,
-		++LayerId,
-		AllottedGeometry.ToPaintGeometry(),
-		&BackgroundBrush,
-		ESlateDrawEffect::None, 
-		BackgroundColor.GetSpecifiedColor()
-	);
+	if (!bHideBackground)
+	{
+		FSlateDrawElement::MakeBox(
+			OutDrawElements,
+			++LayerId,
+			AllottedGeometry.ToPaintGeometry(),
+			&BackgroundBrush,
+			ESlateDrawEffect::None,
+			BackgroundColor.GetSpecifiedColor()
+		);
+	}
+
 
 	if (PixelWidth > 0)
 	{
@@ -56,10 +61,10 @@ int32 SFixedSampledSequenceViewer::OnPaint(const FPaintArgs& Args, const FGeomet
 			TArray<FVector2D> BinDrawPoints;
 			BinDrawPoints.SetNumUninitialized(2);
 
-			for (const SampledSequenceDrawingUtils::FSampleBinCoordinates& BinCoordinates : CachedBinsDrawCoordinates)
+			for (const SampledSequenceDrawingUtils::F2DLineCoordinates& BinCoordinates : CachedBinsDrawCoordinates)
 			{
-				BinDrawPoints[0] = BinCoordinates.Top;
-				BinDrawPoints[1] = BinCoordinates.Bottom;
+				BinDrawPoints[0] = BinCoordinates.A;
+				BinDrawPoints[1] = BinCoordinates.B;
 
 				FSlateDrawElement::MakeLines(
 					OutDrawElements,
@@ -216,7 +221,7 @@ void SFixedSampledSequenceViewer::DrawGridLines(const FGeometry& AllottedGeometr
 
 	for (uint16 Channel = 0; Channel < NChannels; ++Channel)
 	{
-		const SampledSequenceDrawingUtils::FHorizontalDimensionSlot ChannelBoundaries(Channel, NChannels, AllottedGeometry);
+		const SampledSequenceDrawingUtils::FDimensionSlot ChannelBoundaries(Channel, NChannels, AllottedGeometry, DrawingParams);
 
 		LinePoints[0] = FVector2D(0.f, ChannelBoundaries.Center);
 		LinePoints[1] = FVector2D(AllottedGeometry.Size.X, ChannelBoundaries.Center);
diff --git a/Engine/Plugins/Runtime/AudioWidgets/Source/AudioWidgets/Private/SSampledSequenceValueGridOverlay.cpp b/Engine/Plugins/Runtime/AudioWidgets/Source/AudioWidgets/Private/SSampledSequenceValueGridOverlay.cpp
new file mode 100644
--- /dev/null
+++ b/Engine/Plugins/Runtime/AudioWidgets/Source/AudioWidgets/Private/SSampledSequenceValueGridOverlay.cpp
@@ -0,0 +1,278 @@
+// Copyright Epic Games, Inc. All Rights Reserved.
+
+#include "SSampledSequenceValueGridOverlay.h"
+
+#include "Fonts/FontMeasure.h"
+#include "Framework/Application/SlateApplication.h"
+
+void SSampledSequenceValueGridOverlay::Construct(const FArguments& InArgs)
+{
+	check(InArgs._Style);
+
+	GridSplitMode = InArgs._DivideMode;
+	MaxDivisionParameter = InArgs._MaxDivisionParameter;
+	OnValueGridLabel = InArgs._ValueGridLabelGenerator;
+	NumDimensions = InArgs._NumDimensions;
+	DrawingParams = InArgs._SequenceDrawingParams;
+	bHideLabels = InArgs._HideLabels;
+	bHideGrid = InArgs._HideGrid;
+
+	Style = InArgs._Style;
+	GridColor = Style->GridColor;
+	LabelTextFont = Style->LabelTextFont;
+	LabelTextColor = Style->LabelTextColor;
+	DesiredWidth = Style->DesiredWidth;
+	DesiredHeight = Style->DesiredHeight;
+	GridThickness = Style->GridThickness;
+}
+
+int32 SSampledSequenceValueGridOverlay::OnPaint(const FPaintArgs& Args, const FGeometry& AllottedGeometry, const FSlateRect& MyCullingRect, FSlateWindowElementList& OutDrawElements, int32 LayerId, const FWidgetStyle& InWidgetStyle, bool bParentEnabled) const
+{
+	if (bHideGrid && bHideLabels)
+	{
+		return LayerId;
+	}
+
+	TArray<FVector2D> LinePoints;
+	LinePoints.SetNumUninitialized(2);
+	TSharedRef<FSlateFontMeasure> FontMeasureService = FSlateApplication::Get().GetRenderer()->GetFontMeasureService();
+
+	for (int32 SlotIndex = 0; SlotIndex < CachedGridSlotData.Num(); ++SlotIndex)
+	{
+		const SampledSequenceDrawingUtils::FGridData& SlotData = CachedGridSlotData[SlotIndex];
+		const FGridLabelData& LabelData = CachedLabelData[SlotIndex];
+
+		for (int32 LineIndex = 0; LineIndex < SlotData.DrawCoordinates.Num(); ++LineIndex)
+		{
+			LinePoints[0] = SlotData.DrawCoordinates[LineIndex].A;
+			LinePoints[1] = SlotData.DrawCoordinates[LineIndex].B;
+
+			if (!bHideLabels)
+			{
+				const FVector2D LabelSize = FontMeasureService->Measure(LabelData.LabelTexts[LineIndex], FSlateFontInfo(FAppStyle::GetFontStyle("Regular")));
+				const FSlateLayoutTransform LabelLayoutTransform(1.0f, LabelData.LabelCoordinates[LineIndex]);
+				float LabelRotation = 0.f;
+
+				switch (DrawingParams.Orientation)
+				{
+				case SampledSequenceDrawingUtils::ESampledSequenceDrawOrientation::Horizontal:
+					LabelRotation = 0.f;
+					break;
+				case SampledSequenceDrawingUtils::ESampledSequenceDrawOrientation::Vertical:
+					LabelRotation = 90.f;
+				default:
+					static_assert(static_cast<int32>(SampledSequenceDrawingUtils::ESampledSequenceDrawOrientation::COUNT) == 2, "Possible missing switch case coverage for 'ESampledSequenceDrawOrientation'");
+					break;
+				}
+				
+				const FSlateRenderTransform LabelRotationTransform = FSlateRenderTransform(FQuat2D(FMath::DegreesToRadians(LabelRotation)));
+
+
+				const FGeometry LabelGeometry = AllottedGeometry.MakeChild(
+					LabelSize,
+					LabelLayoutTransform,
+					LabelRotationTransform,
+					FVector2D(0.0f, 0.0f)
+				);
+
+				FSlateDrawElement::MakeText(
+					OutDrawElements,
+					LayerId,
+					LabelGeometry.ToPaintGeometry(),
+					LabelData.LabelTexts[LineIndex],
+					LabelTextFont,
+					ESlateDrawEffect::None,
+					LabelTextColor.GetSpecifiedColor()
+				);
+			}
+
+			if (!bHideGrid)
+			{
+				FSlateDrawElement::MakeLines(
+					OutDrawElements,
+					LayerId,
+					AllottedGeometry.ToPaintGeometry(),
+					LinePoints,
+					ESlateDrawEffect::None,
+					GridColor.GetSpecifiedColor(),
+					false,
+					GridThickness
+				);
+			}
+		}
+	}
+
+	return ++LayerId;
+}
+
+FVector2D SSampledSequenceValueGridOverlay::ComputeDesiredSize(float) const
+{
+	return FVector2D(DesiredWidth, DesiredHeight);
+}
+
+void SSampledSequenceValueGridOverlay::Tick(const FGeometry& AllottedGeometry, const double InCurrentTime, const float InDeltaTime)
+{
+	if (AllottedGeometry.GetLocalSize() != CachedLocalSize || bForceRedraw)
+	{
+		CachedLocalSize = AllottedGeometry.GetLocalSize();
+		CacheDrawElements(AllottedGeometry, MaxDivisionParameter);
+		bForceRedraw = false;
+	}
+}
+
+void SSampledSequenceValueGridOverlay::OnStyleUpdated(const FNotifyingAudioWidgetStyle& UpdatedStyle)
+{
+	check(Style);
+
+	if (&UpdatedStyle != Style)
+	{
+		return;
+	}
+
+	GridColor = Style->GridColor;
+	LabelTextFont = Style->LabelTextFont;
+	LabelTextColor = Style->LabelTextColor;
+	DesiredWidth = Style->DesiredWidth;
+	DesiredHeight = Style->DesiredHeight;
+	GridThickness = Style->GridThickness;
+
+	ForceRedraw();
+}
+
+void SSampledSequenceValueGridOverlay::SetMaxDivisionParameter(const uint32 InDivisionParameter)
+{
+	MaxDivisionParameter = InDivisionParameter;
+	ForceRedraw();
+}
+
+void SSampledSequenceValueGridOverlay::SetHideLabels(const bool InHideLabels)
+{
+	bHideLabels = InHideLabels;
+	ForceRedraw();
+}
+
+void SSampledSequenceValueGridOverlay::SetHideGrid(const bool InHideGrid)
+{
+	bHideGrid = InHideGrid;
+	ForceRedraw();
+}
+
+void SSampledSequenceValueGridOverlay::ForceRedraw() 
+{
+	bForceRedraw = true;
+}
+
+void SSampledSequenceValueGridOverlay::CacheDrawElements(const FGeometry& AllottedGeometry, const uint32 InDivisionParameter)
+{
+	switch (GridSplitMode)
+	{
+	case SampledSequenceValueGridOverlay::EGridDivideMode::EvenSplit:
+		SampledSequenceDrawingUtils::GenerateEvenlySplitGridForGeometry(CachedGridSlotData, AllottedGeometry, NumDimensions, InDivisionParameter, DrawingParams);
+		break;
+	case SampledSequenceValueGridOverlay::EGridDivideMode::MidSplit:
+		SampledSequenceDrawingUtils::GenerateMidpointSplitGridForGeometry(CachedGridSlotData, AllottedGeometry, NumDimensions, InDivisionParameter, DrawingParams);
+		break;
+	default:
+		static_assert(static_cast<int32>(SampledSequenceValueGridOverlay::EGridDivideMode::COUNT) == 2, "Possible missing switch case coverage for 'ValueGridOverlay::EGridDivideMode'");
+		break;
+	}
+	
+	if (bHideLabels)
+	{
+		return;
+	}
+
+	switch (DrawingParams.Orientation)
+	{
+	case SampledSequenceDrawingUtils::ESampledSequenceDrawOrientation::Horizontal:
+		GenerateHorizontalGridLabels(InDivisionParameter, AllottedGeometry);
+		break;
+	case SampledSequenceDrawingUtils::ESampledSequenceDrawOrientation::Vertical:
+		GenerateVerticalGridLabels(InDivisionParameter, AllottedGeometry);
+		break;
+	default:
+		static_assert(static_cast<int32>(SampledSequenceDrawingUtils::ESampledSequenceDrawOrientation::COUNT) == 2, "Possible missing switch case coverage for 'ESampledSequenceDrawOrientation'");
+		break;
+	}
+	return;
+}
+
+void SSampledSequenceValueGridOverlay::GenerateVerticalGridLabels(const uint32 InDivisionParameter, const FGeometry& AllottedGeometry)
+{
+	CachedLabelData.Empty();
+	TSharedRef<FSlateFontMeasure> FontMeasureService = FSlateApplication::Get().GetRenderer()->GetFontMeasureService();
+
+	for (int32 SlotIndex = 0; SlotIndex < CachedGridSlotData.Num(); ++SlotIndex)
+	{
+		FGridLabelData LabelData;
+		const SampledSequenceDrawingUtils::FGridData& SlotData = CachedGridSlotData[SlotIndex];
+
+		for (int32 LineIndex = 0; LineIndex < SlotData.DrawCoordinates.Num(); ++LineIndex)
+		{
+			FText LabelText = OnValueGridLabel(SlotData.PositionRatios[LineIndex]);
+			const FVector2D TextSize = FontMeasureService->Measure(LabelText, LabelTextFont);
+			const FVector2D& DrawCoordinate = SlotData.DrawCoordinates[LineIndex].A;
+
+			const float TextX = SlotData.PositionRatios[LineIndex] <= 0.5 ? DrawCoordinate.X + LabelToGridPixelDistance + TextSize.Y : DrawCoordinate.X - LabelToGridPixelDistance;
+			const FVector2D TextOffset = FVector2D(TextX, DrawCoordinate.Y + LabelToGridPixelDistance);
+
+			if (SlotData.DrawCoordinates.Num() > 1)
+			{
+				const int32 ComparedSlotIndex = LineIndex == SlotData.DrawCoordinates.Num() - 1 ? LineIndex - 1 : LineIndex + 1;
+				const float GridLinesSpacing = FMath::Abs(DrawCoordinate.X - SlotData.DrawCoordinates[ComparedSlotIndex].A.X);
+
+				if (GridLinesSpacing < TextSize.Y * 2 && InDivisionParameter != 0)
+				{
+					const float NewDivisionParameter = FMath::Clamp(InDivisionParameter - 1, 0, MaxDivisionParameter);
+					CacheDrawElements(AllottedGeometry, NewDivisionParameter);
+					return;
+				}
+			}
+
+			LabelData.LabelTexts.Emplace(LabelText);
+			LabelData.LabelCoordinates.Emplace(TextOffset);
+		}
+
+		CachedLabelData.Emplace(LabelData);
+	}
+}
+
+void SSampledSequenceValueGridOverlay::GenerateHorizontalGridLabels(const uint32 InDivisionParameter, const FGeometry& AllottedGeometry)
+{
+	CachedLabelData.Empty();
+	TSharedRef<FSlateFontMeasure> FontMeasureService = FSlateApplication::Get().GetRenderer()->GetFontMeasureService();
+
+	for (int32 SlotIndex = 0; SlotIndex < CachedGridSlotData.Num(); ++SlotIndex)
+	{
+		FGridLabelData LabelData;
+		const SampledSequenceDrawingUtils::FGridData& SlotData = CachedGridSlotData[SlotIndex];
+
+		for (int32 LineIndex = 0; LineIndex < SlotData.DrawCoordinates.Num(); ++LineIndex)
+		{
+			FText LabelText = OnValueGridLabel(SlotData.PositionRatios[LineIndex]);
+			const FVector2D TextSize = FontMeasureService->Measure(LabelText, LabelTextFont);
+			const FVector2D& DrawCoordinate = SlotData.DrawCoordinates[LineIndex].A;
+
+			const float TextY = SlotData.PositionRatios[LineIndex] <= 0.5 ? DrawCoordinate.Y + LabelToGridPixelDistance : DrawCoordinate.Y - TextSize.Y - LabelToGridPixelDistance;
+			const FVector2D TextOffset = FVector2D(DrawCoordinate.X + LabelToGridPixelDistance, TextY);
+
+			if (SlotData.DrawCoordinates.Num() > 1)
+			{
+				const int32 ComparedSlotIndex = LineIndex == SlotData.DrawCoordinates.Num() - 1 ? LineIndex - 1 : LineIndex + 1;
+				const float GridLinesSpacing = FMath::Abs(DrawCoordinate.Y - SlotData.DrawCoordinates[ComparedSlotIndex].A.Y);
+
+				if (GridLinesSpacing < TextSize.Y * 2 && InDivisionParameter != 0)
+				{
+					const float NewDivisionParameter = FMath::Clamp(InDivisionParameter - 1, 0, MaxDivisionParameter);
+					CacheDrawElements(AllottedGeometry, NewDivisionParameter);
+					return;
+				}
+			}
+
+			LabelData.LabelTexts.Emplace(LabelText);
+			LabelData.LabelCoordinates.Emplace(TextOffset);
+		}
+
+		CachedLabelData.Emplace(LabelData);
+	}
+}
\ No newline at end of file
diff --git a/Engine/Plugins/Runtime/AudioWidgets/Source/AudioWidgets/Private/SampledSequenceDrawingUtils.cpp b/Engine/Plugins/Runtime/AudioWidgets/Source/AudioWidgets/Private/SampledSequenceDrawingUtils.cpp
--- a/Engine/Plugins/Runtime/AudioWidgets/Source/AudioWidgets/Private/SampledSequenceDrawingUtils.cpp
+++ b/Engine/Plugins/Runtime/AudioWidgets/Source/AudioWidgets/Private/SampledSequenceDrawingUtils.cpp
@@ -4,17 +4,61 @@
 
 #include "Layout/Geometry.h"
 
-SampledSequenceDrawingUtils::FHorizontalDimensionSlot::FHorizontalDimensionSlot(const uint16 DimensionToDraw, const uint16 TotalNumDimensions, const FGeometry& InAllottedGeometry)
+namespace SampledSequenceDrawingUtilsPrivate
 {
-	Height = InAllottedGeometry.GetLocalSize().Y / TotalNumDimensions;
-	float ChannelSlotMidPoint = Height / 2.f;
+	void DivideLineIteratively(const float StartPoint, const float EndPoint, int32 Depth, TArray<double>& HalfPoints)
+	{
+		const double HalfPoint = (StartPoint + EndPoint) / 2.f;
+		HalfPoints.Add(HalfPoint);
+
+		if (--Depth > 0)
+		{
+			SampledSequenceDrawingUtilsPrivate::DivideLineIteratively(StartPoint, HalfPoint, Depth, HalfPoints);
+			SampledSequenceDrawingUtilsPrivate::DivideLineIteratively(HalfPoint, EndPoint, Depth, HalfPoints);
+		}
+	}
+
+	void GenerateGridDataFromMidPoints(const uint16 NDimensions, const FGeometry& InAllottedGeometry, const SampledSequenceDrawingUtils::FSampledSequenceDrawingParams& Params, const TArray<double>& MidPoints, const TArray<float>& MidPointsRatios, TArray<SampledSequenceDrawingUtils::FGridData>& OutGridData)
+	{
+		for (uint16 Dimension = 0; Dimension < NDimensions; ++Dimension)
+		{
+			SampledSequenceDrawingUtils::FGridData DimensionGridData;
+			SampledSequenceDrawingUtils::FDimensionSlot DimensionSlot(Dimension, NDimensions, InAllottedGeometry, Params);
+
+			for (int32 PointIndex = 0; PointIndex < MidPoints.Num(); ++PointIndex)
+			{
+				const float MidPointCoordinate = MidPoints[PointIndex] + DimensionSlot.Top;
+
+				FVector2D PointA = Params.Orientation == SampledSequenceDrawingUtils::ESampledSequenceDrawOrientation::Horizontal ? FVector2D(0, MidPointCoordinate) : FVector2D(MidPointCoordinate, 0);
+				FVector2D PointB = Params.Orientation == SampledSequenceDrawingUtils::ESampledSequenceDrawOrientation::Horizontal ? FVector2D(InAllottedGeometry.GetLocalSize().X, MidPointCoordinate) : FVector2D(MidPointCoordinate, InAllottedGeometry.GetLocalSize().Y);
 
-	Top = Height * DimensionToDraw;
-	Center = Top + ChannelSlotMidPoint;
-	Bottom = Top + Height;
+				DimensionGridData.DrawCoordinates.Add({ PointA, PointB });
+				DimensionGridData.PositionRatios.Add(MidPointsRatios[PointIndex]);
+			}
+
+			OutGridData.Add(DimensionGridData);
+		}
+	}
 }
 
-void SampledSequenceDrawingUtils::GenerateSampleBinsCoordinatesForGeometry(TArray<FSampleBinCoordinates>& OutDrawCoordinates, const FGeometry& InAllottedGeometry, const TArray<TRange<float>>& InSampleBins, const uint16 NDimensions, const FSampledSequenceDrawingParams Params)
+SampledSequenceDrawingUtils::FDimensionSlot::FDimensionSlot(const uint16 DimensionToDraw, const uint16 TotalNumDimensions, const FGeometry& InAllottedGeometry, const FSampledSequenceDrawingParams& Params)
+{
+	check(TotalNumDimensions > 0)
+
+	const float& GeometryLength = Params.Orientation == ESampledSequenceDrawOrientation::Horizontal ? InAllottedGeometry.GetLocalSize().Y : InAllottedGeometry.GetLocalSize().X;
+
+	const double FullSlotHeight = GeometryLength / TotalNumDimensions;
+	const double MarginedTop = (FullSlotHeight * DimensionToDraw) + Params.DimensionSlotMargin;
+	Height = FMath::Clamp(FullSlotHeight - (Params.DimensionSlotMargin * 2), FullSlotHeight * Params.MinSequenceHeightRatio, GeometryLength / TotalNumDimensions);
+	const double MarginedBottom = MarginedTop + Height;
+
+	Center = MarginedTop + ((MarginedBottom - MarginedTop) / 2.f);
+	Height *= Params.MaxSequenceHeightRatio;
+	Top = Center - Height / 2;
+	Bottom = Center + Height / 2;
+}
+
+void SampledSequenceDrawingUtils::GenerateSampleBinsCoordinatesForGeometry(TArray<F2DLineCoordinates>& OutDrawCoordinates, const FGeometry& InAllottedGeometry, const TArray<TRange<float>>& InSampleBins, const uint16 NDimensions, const FSampledSequenceDrawingParams Params)
 {
 	if (!ensure(NDimensions != 0))
 	{
@@ -24,29 +68,29 @@ void SampledSequenceDrawingUtils::GenerateSampleBinsCoordinatesForGeometry(TArra
 	const uint32 PixelWidth = FMath::FloorToInt(InAllottedGeometry.GetLocalSize().X);
 	check(PixelWidth * NDimensions == InSampleBins.Num());
 
-	const float HeightScale = InAllottedGeometry.GetLocalSize().Y / (2.f * Params.MaxDisplayedValue * NDimensions) * Params.MaxSequenceHeight * Params.VerticalZoomFactor;
+	const float HeightScale = InAllottedGeometry.GetLocalSize().Y / (2.f * Params.MaxDisplayedValue * NDimensions) * Params.MaxSequenceHeightRatio * Params.VerticalZoomFactor;
 
 
 	OutDrawCoordinates.SetNumUninitialized(InSampleBins.Num());
-	FSampleBinCoordinates* OutCoordinatesData = OutDrawCoordinates.GetData();
+	F2DLineCoordinates* OutCoordinatesData = OutDrawCoordinates.GetData();
 
 	for (uint16 Channel = 0; Channel < NDimensions; ++Channel)
 	{
-		const FHorizontalDimensionSlot ChannelBoundaries(Channel, NDimensions, InAllottedGeometry);
+		const FDimensionSlot ChannelBoundaries(Channel, NDimensions, InAllottedGeometry, Params);
 
 		for (uint32 Pixel = 0; Pixel < PixelWidth; ++Pixel)
 		{
 			uint32 PeakIndex = Pixel * NDimensions + Channel;
 			const TRange<float>& MinMaxBin = InSampleBins[PeakIndex];
 
-			const float SampleMaxScaled = MinMaxBin.GetUpperBoundValue() * HeightScale > Params.MinScaledBinValue ? MinMaxBin.GetUpperBoundValue() * HeightScale : Params.MinSequenceHeight;
-			const float SampleMinScaled = MinMaxBin.GetLowerBoundValue() * HeightScale < -1 * Params.MinScaledBinValue ? MinMaxBin.GetLowerBoundValue() * HeightScale : -1.f * Params.MinSequenceHeight;
+			const float SampleMaxScaled = MinMaxBin.GetUpperBoundValue() * HeightScale > Params.MinScaledBinValue ? MinMaxBin.GetUpperBoundValue() * HeightScale : Params.MinSequenceHeightRatio;
+			const float SampleMinScaled = MinMaxBin.GetLowerBoundValue() * HeightScale < -1 * Params.MinScaledBinValue ? MinMaxBin.GetLowerBoundValue() * HeightScale : -1.f * Params.MinSequenceHeightRatio;
 
-			const float Top = FMath::Max(ChannelBoundaries.Center - SampleMaxScaled, ChannelBoundaries.Top + Params.DimensionSlotMargin);
-			const float Bottom = FMath::Min(ChannelBoundaries.Center - SampleMinScaled, ChannelBoundaries.Bottom - Params.DimensionSlotMargin);
+			const float Top = FMath::Max(ChannelBoundaries.Center - SampleMaxScaled, ChannelBoundaries.Top);
+			const float Bottom = FMath::Min(ChannelBoundaries.Center - SampleMinScaled, ChannelBoundaries.Bottom);
 
-			OutCoordinatesData[PeakIndex].Top = FVector2D(Pixel, Top);
-			OutCoordinatesData[PeakIndex].Bottom = FVector2D(Pixel, Bottom);
+			OutCoordinatesData[PeakIndex].A = FVector2D(Pixel, Top);
+			OutCoordinatesData[PeakIndex].B = FVector2D(Pixel, Bottom);
 		}
 	}
 }
@@ -62,11 +106,11 @@ void SampledSequenceDrawingUtils::GenerateSequencedSamplesCoordinatesForGeometry
 
 	const uint32 NumFramesToDisplay = InSampleData.Num() / NDimensions;
 
-	TArray<FHorizontalDimensionSlot> ChannelSlotsBoundaries;
+	TArray<FDimensionSlot> ChannelSlotsBoundaries;
 
 	for (uint16 Channel = 0; Channel < NDimensions; ++Channel)
 	{
-		ChannelSlotsBoundaries.Emplace(Channel, NDimensions, InAllottedGeometry);
+		ChannelSlotsBoundaries.Emplace(Channel, NDimensions, InAllottedGeometry, Params);
 	}
 
 	uint32 FrameIndex = 0;
@@ -80,14 +124,13 @@ void SampledSequenceDrawingUtils::GenerateSequencedSamplesCoordinatesForGeometry
 		{
 			for (uint16 Channel = 0; Channel < NDimensions; ++Channel)
 			{
-				const FHorizontalDimensionSlot& ChannelBoundaries = ChannelSlotsBoundaries[Channel];
+				const FDimensionSlot& ChannelBoundaries = ChannelSlotsBoundaries[Channel];
 				const uint32 SampleIndex = FrameIndex * NDimensions + Channel;
 				const float SampleValue = InSampleData[SampleIndex];
 
-				const float SampleValueRatio = SampleValue / Params.MaxDisplayedValue * Params.MaxSequenceHeight * Params.VerticalZoomFactor;
-				const float TopBoundary = ChannelBoundaries.Top + Params.DimensionSlotMargin;
-				const float BottomBoundary = ChannelBoundaries.Bottom - Params.DimensionSlotMargin;
-				const float SampleY = FMath::Clamp((SampleValueRatio * ChannelBoundaries.Height / 2.f) + ChannelBoundaries.Center, TopBoundary, BottomBoundary);
+				const float SampleValueRatio = SampleValue / Params.MaxDisplayedValue * Params.VerticalZoomFactor;
+				const float HeightRatio = SampleValueRatio * ChannelBoundaries.Height / 2.f;
+				const float SampleY = FMath::Clamp(ChannelBoundaries.Center - HeightRatio, ChannelBoundaries.Top, ChannelBoundaries.Bottom);
 
 				const FVector2D SampleCoordinates(FrameX, SampleY);
 
@@ -98,3 +141,96 @@ void SampledSequenceDrawingUtils::GenerateSequencedSamplesCoordinatesForGeometry
 		}
 	}
 }
+
+void SampledSequenceDrawingUtils::GenerateEvenlySplitGridForGeometry(TArray<FGridData>& OutGridData, const FGeometry& InAllottedGeometry, const uint16 NDimensions, const uint32 NumGridDivisions /*= 5*/, const FSampledSequenceDrawingParams Params /*= FSampledSequenceDrawingParams()*/)
+{
+	if (!ensure(NDimensions != 0))
+	{
+		return;
+	}
+
+	OutGridData.Empty();
+
+	FDimensionSlot FirstDimensionSlot(0, NDimensions, InAllottedGeometry, Params);
+	TArray<double> MidPoints;
+	TArray<float> MidPointsRatios;
+	SampledSequenceDrawingUtils::GenerateEvenlySplitGridForLine(MidPoints, MidPointsRatios, FirstDimensionSlot.Height, NumGridDivisions);
+	SampledSequenceDrawingUtilsPrivate::GenerateGridDataFromMidPoints(NDimensions, InAllottedGeometry, Params, MidPoints, MidPointsRatios, OutGridData);
+
+}
+
+void SampledSequenceDrawingUtils::GenerateMidpointSplitGridForGeometry(TArray<FGridData>& OutGridData, const FGeometry& InAllottedGeometry, const uint16 NDimensions, const uint32 DivisionDepth /*= 3*/, const FSampledSequenceDrawingParams Params /*= FSampledSequenceDrawingParams()*/)
+{
+	if (!ensure(NDimensions != 0))
+	{
+		return;
+	}
+
+	OutGridData.Empty();
+
+	FDimensionSlot FirstDimensionSlot(0, NDimensions, InAllottedGeometry, Params);
+	TArray<double> MidPoints;
+	TArray<float> MidPointsRatios;
+	SampledSequenceDrawingUtils::GenerateMidpointSplitGridForLine(MidPoints, MidPointsRatios, FirstDimensionSlot.Height, DivisionDepth);
+	SampledSequenceDrawingUtilsPrivate::GenerateGridDataFromMidPoints(NDimensions, InAllottedGeometry, Params, MidPoints, MidPointsRatios, OutGridData);
+}
+
+
+void SampledSequenceDrawingUtils::GenerateMidpointSplitGridForLine(TArray<double>& OutDrawCoordinates, TArray<float>& OutLinePositionRatios, const float LineLength, const uint32 DivisionDepth /*= 3*/, bool bEmptyOutArrays /*= true*/)
+{
+	if (bEmptyOutArrays)
+	{
+		OutDrawCoordinates.Empty();
+		OutLinePositionRatios.Empty();
+	}
+
+	const bool bAddBorders = DivisionDepth > 0;
+	
+	if (bAddBorders)
+	{
+		OutDrawCoordinates.Add(0);
+		OutLinePositionRatios.Add(0.f);
+	}
+
+
+	TArray<double> MidPoints;
+	SampledSequenceDrawingUtilsPrivate::DivideLineIteratively(0, LineLength, DivisionDepth, MidPoints);
+
+	MidPoints.Sort();
+
+	for (const double MidPoint : MidPoints)
+	{
+		const float MidPointRatio = MidPoint / LineLength;
+		OutLinePositionRatios.Add(MidPointRatio);
+		OutDrawCoordinates.Add(MidPoint);
+	}
+
+	if (bAddBorders)
+	{
+		OutDrawCoordinates.Add(LineLength);
+		OutLinePositionRatios.Add(1);
+	}
+
+	
+}
+
+void SampledSequenceDrawingUtils::GenerateEvenlySplitGridForLine(TArray<double>& OutDrawCoordinates, TArray<float>& OutLinePositionRatios, const float LineLength, const uint32 NumGridDivisions /*= 2*/, bool bEmptyOutArrays /*= true*/)
+{
+	if (bEmptyOutArrays)
+	{
+		OutDrawCoordinates.Empty();
+		OutLinePositionRatios.Empty();
+	}
+
+	const double GridStep = LineLength / NumGridDivisions;
+
+	for (double GridLineCoordinate = 0; GridLineCoordinate < LineLength; GridLineCoordinate += GridStep)
+	{
+		const float GridValue = GridLineCoordinate / LineLength;
+		OutLinePositionRatios.Add(GridValue);
+		OutDrawCoordinates.Add(GridLineCoordinate);
+	}
+
+	OutLinePositionRatios.Add(1.f);
+	OutDrawCoordinates.Add(LineLength);
+}
diff --git a/Engine/Plugins/Runtime/AudioWidgets/Source/AudioWidgets/Public/AudioWidgetsSlateTypes.h b/Engine/Plugins/Runtime/AudioWidgets/Source/AudioWidgets/Public/AudioWidgetsSlateTypes.h
--- a/Engine/Plugins/Runtime/AudioWidgets/Source/AudioWidgets/Public/AudioWidgetsSlateTypes.h
+++ b/Engine/Plugins/Runtime/AudioWidgets/Source/AudioWidgets/Public/AudioWidgetsSlateTypes.h
@@ -341,6 +341,48 @@ struct AUDIOWIDGETS_API FFixedSampleSequenceRulerStyle : public FNotifyingAudioW
 	FFixedSampleSequenceRulerStyle& SetDesiredHeight(const float InDesiredHeight) { DesiredHeight = InDesiredHeight; BroadcastStyleUpdate(); return *this; }
 };
 
+/**
+ * Represents the appearance of a Sampled Sequence Value Grid Overlay
+ */
+USTRUCT(BlueprintType)
+struct AUDIOWIDGETS_API FSampledSequenceValueGridOverlayStyle : public FNotifyingAudioWidgetStyle
+{
+	GENERATED_USTRUCT_BODY()
+
+	FSampledSequenceValueGridOverlayStyle();
+
+	static const FSampledSequenceValueGridOverlayStyle& GetDefault();
+	virtual const FName GetTypeName() const override { return TypeName; };
+
+	static const FName TypeName;
+
+	UPROPERTY(EditAnywhere, BlueprintReadWrite, Category = Appearance)
+	FSlateColor GridColor;
+	FSampledSequenceValueGridOverlayStyle& SetGridColor(const FSlateColor InGridColor) { GridColor = InGridColor; BroadcastStyleUpdate(); return *this; }
+
+	UPROPERTY(EditAnywhere, BlueprintReadWrite, Category = Appearance)
+	float GridThickness;
+	FSampledSequenceValueGridOverlayStyle& SetGridThickness(const float InGridThickness) { GridThickness = InGridThickness; BroadcastStyleUpdate(); return *this; }
+
+	UPROPERTY(EditAnywhere, BlueprintReadWrite, Category = Appearance)
+	FSlateColor LabelTextColor;
+	FSampledSequenceValueGridOverlayStyle& SetLabelTextColor(const FSlateColor& InLabelTextColor) { LabelTextColor = InLabelTextColor; BroadcastStyleUpdate(); return *this; }
+
+	UPROPERTY(EditAnywhere, BlueprintReadWrite, Category = Appearance)
+	FSlateFontInfo LabelTextFont;
+	FSampledSequenceValueGridOverlayStyle& SetLabelTextFont(const FSlateFontInfo& InLabelTextFont) { LabelTextFont = InLabelTextFont; BroadcastStyleUpdate(); return *this; }
+	FSampledSequenceValueGridOverlayStyle& SetLabelTextFontSize(const float InFontSize) { LabelTextFont.Size = InFontSize; BroadcastStyleUpdate(); return *this; }
+
+	UPROPERTY(EditAnywhere, BlueprintReadWrite, Category = Appearance)
+	float DesiredWidth;
+	FSampledSequenceValueGridOverlayStyle& SetDesiredWidth(const float InDesiredWidth) { DesiredWidth = InDesiredWidth; BroadcastStyleUpdate(); return *this; }
+
+	UPROPERTY(EditAnywhere, BlueprintReadWrite, Category = Appearance)
+	float DesiredHeight;
+	FSampledSequenceValueGridOverlayStyle& SetDesiredHeight(const float InDesiredHeight) { DesiredHeight = InDesiredHeight; BroadcastStyleUpdate(); return *this; }
+
+};
+
 #if UE_ENABLE_INCLUDE_ORDER_DEPRECATED_IN_5_2
 #include "Styling/SlateStyle.h"
 #include "Styling/SlateStyleMacros.h"
diff --git a/Engine/Plugins/Runtime/AudioWidgets/Source/AudioWidgets/Public/SFixedSampledSequenceViewer.h b/Engine/Plugins/Runtime/AudioWidgets/Source/AudioWidgets/Public/SFixedSampledSequenceViewer.h
--- a/Engine/Plugins/Runtime/AudioWidgets/Source/AudioWidgets/Public/SFixedSampledSequenceViewer.h
+++ b/Engine/Plugins/Runtime/AudioWidgets/Source/AudioWidgets/Public/SFixedSampledSequenceViewer.h
@@ -19,11 +19,14 @@ class AUDIOWIDGETS_API SFixedSampledSequenceViewer : public SLeafWidget
 public:
 
 	SLATE_BEGIN_ARGS(SFixedSampledSequenceViewer) 
+		: _HideBackground(false)
 	{
 	}
 	
 	SLATE_ARGUMENT(SampledSequenceDrawingUtils::FSampledSequenceDrawingParams, SequenceDrawingParams)
 
+	SLATE_ARGUMENT(bool, HideBackground)
+
 	SLATE_STYLE_ARGUMENT(FSampledSequenceViewerStyle, Style)
 		
 	SLATE_END_ARGS()
@@ -63,12 +66,12 @@ private:
 	float DesiredWidth = 0.f;
 	float SampleMarkersSize = 2.5f;
 	float SequenceLineThickness = 1.f;
-	
+	bool bHideBackground;
 
 	uint32 CachedPixelWidth = 0; 
 	float CachedPixelHeight = 0.f;
 	TArray<SampledSequenceViewerWidget::SampleRange> CachedPeaks;
-	TArray<SampledSequenceDrawingUtils::FSampleBinCoordinates> CachedBinsDrawCoordinates;
+	TArray<SampledSequenceDrawingUtils::F2DLineCoordinates> CachedBinsDrawCoordinates;
 	SampledSequenceDrawingUtils::FSampledSequenceDrawingParams DrawingParams;
 	TArray<FVector2D> CachedSampleDrawCoordinates;
 
diff --git a/Engine/Plugins/Runtime/AudioWidgets/Source/AudioWidgets/Public/SSampledSequenceValueGridOverlay.h b/Engine/Plugins/Runtime/AudioWidgets/Source/AudioWidgets/Public/SSampledSequenceValueGridOverlay.h
new file mode 100644
--- /dev/null
+++ b/Engine/Plugins/Runtime/AudioWidgets/Source/AudioWidgets/Public/SSampledSequenceValueGridOverlay.h
@@ -0,0 +1,110 @@
+// Copyright Epic Games, Inc. All Rights Reserved.
+
+#pragma once
+
+#include "AudioWidgetsSlateTypes.h"
+#include "SampledSequenceDrawingUtils.h"
+#include "Styling/ISlateStyle.h"
+#include "Styling/SlateWidgetStyleAsset.h"
+#include "Widgets/SLeafWidget.h"
+
+namespace SampledSequenceValueGridOverlay
+{
+	enum class EGridDivideMode
+	{
+		EvenSplit = 0,
+		MidSplit,
+		COUNT
+	};
+
+	TFunction<FText(const double)> DefaultLabelGenerator = [](const double GridLineValue)
+	{
+		FNumberFormattingOptions DefaultFormatting;
+		return FText::AsNumber(GridLineValue, &DefaultFormatting);
+	};
+}
+
+class AUDIOWIDGETS_API SSampledSequenceValueGridOverlay : public SLeafWidget
+{
+public:
+	SLATE_BEGIN_ARGS(SSampledSequenceValueGridOverlay)
+		: _DivideMode(SampledSequenceValueGridOverlay::EGridDivideMode::EvenSplit)
+		, _NumDimensions(1)
+		, _HideLabels(false)
+		, _HideGrid(false)
+		, _MaxDivisionParameter(2)
+		, _ValueGridLabelGenerator(SampledSequenceValueGridOverlay::DefaultLabelGenerator)
+	{
+	}
+
+	SLATE_ARGUMENT(SampledSequenceValueGridOverlay::EGridDivideMode, DivideMode)
+
+	SLATE_ARGUMENT(uint32, NumDimensions)
+
+	SLATE_ARGUMENT(bool, HideLabels)
+
+	SLATE_ARGUMENT(bool, HideGrid)
+
+	SLATE_ARGUMENT(uint32, MaxDivisionParameter)
+
+	SLATE_ARGUMENT(SampledSequenceDrawingUtils::FSampledSequenceDrawingParams, SequenceDrawingParams)
+
+	SLATE_ARGUMENT(TFunction<FText(const double)>, ValueGridLabelGenerator)
+
+	SLATE_STYLE_ARGUMENT(FSampledSequenceValueGridOverlayStyle, Style)
+
+	SLATE_END_ARGS()
+
+	void Construct(const FArguments& InArgs);
+
+	virtual int32 OnPaint(const FPaintArgs& Args, const FGeometry& AllottedGeometry, const FSlateRect& MyCullingRect, FSlateWindowElementList& OutDrawElements, int32 LayerId, const FWidgetStyle& InWidgetStyle, bool bParentEnabled) const override;
+	virtual FVector2D ComputeDesiredSize(float) const override;
+	virtual void Tick(const FGeometry& AllottedGeometry, const double InCurrentTime, const float InDeltaTime) override;
+	void OnStyleUpdated(const FNotifyingAudioWidgetStyle& UpdatedStyle);
+
+	void SetMaxDivisionParameter(const uint32 InDivisionParameter);
+	void SetHideLabels(const bool InHideLabels);
+	void SetHideGrid(const bool InHideGrid);
+
+	void ForceRedraw();
+
+private:
+	void CacheDrawElements(const FGeometry& AllottedGeometry, const uint32 InDivisionParameter);
+	
+	void GenerateVerticalGridLabels(const uint32 InDivisionParameter, const FGeometry& AllottedGeometry);
+	void GenerateHorizontalGridLabels(const uint32 InDivisionParameter, const FGeometry& AllottedGeometry);
+
+
+	SampledSequenceValueGridOverlay::EGridDivideMode GridSplitMode;
+	uint32 MaxDivisionParameter;
+	uint32 NumDimensions;
+	TFunction<FText(const double)> OnValueGridLabel;
+	SampledSequenceDrawingUtils::FSampledSequenceDrawingParams DrawingParams;
+	FVector2D CachedLocalSize;
+	bool bForceRedraw = false;
+	bool bHideLabels = false;
+	bool bHideGrid = false;
+
+	TArray<SampledSequenceDrawingUtils::FGridData> CachedGridSlotData;
+	FNumberFormattingOptions ValueGridFormatOptions;
+	const float LabelToGridPixelDistance = 2;
+
+	const FSampledSequenceValueGridOverlayStyle* Style = nullptr;
+
+	FSlateColor GridColor;
+	FSlateColor LabelTextColor;
+	FSlateFontInfo LabelTextFont;
+	float DesiredWidth;
+	float DesiredHeight;
+	float GridThickness;
+
+
+	struct FGridLabelData
+	{
+		TArray<FText> LabelTexts;
+		TArray<FVector2D> LabelCoordinates;
+	};
+
+	TArray<FGridLabelData> CachedLabelData;
+
+};
\ No newline at end of file
diff --git a/Engine/Plugins/Runtime/AudioWidgets/Source/AudioWidgets/Public/SampledSequenceDrawingUtils.h b/Engine/Plugins/Runtime/AudioWidgets/Source/AudioWidgets/Public/SampledSequenceDrawingUtils.h
--- a/Engine/Plugins/Runtime/AudioWidgets/Source/AudioWidgets/Public/SampledSequenceDrawingUtils.h
+++ b/Engine/Plugins/Runtime/AudioWidgets/Source/AudioWidgets/Public/SampledSequenceDrawingUtils.h
@@ -96,27 +96,37 @@ namespace SampledSequenceDrawingUtils
 		* Params for drawing sequences of samples/sample bins
 		* @param MaxDisplayedValue					The highest value a sample can take
 		* @param DimensionSlotMargin				Margin to keep from the dimension slot boundaries (pixels)
-		* @param MinSequenceHeight					The minimum height (ratio) the drawn sequence can take in a channel slot
-		* @param MaxSequenceHeight					The maximum height (ratio) the drawn sequence can take in a channel slot
+		* @param MinSequenceHeightRatio				The minimum height (ratio) the drawn sequence can take in a channel slot
+		* @param MaxSequenceHeightRatio				The maximum height (ratio) the drawn sequence can take in a channel slot
 		* @param MinScaledBinValue					Minimum value a scaled bin can have
 		* @param VerticalZoomFactor					Sequence zoom factor
 	*/
+
+	enum class ESampledSequenceDrawOrientation
+	{
+		Horizontal = 0, 
+		Vertical, 
+		COUNT
+	};
+
+
 	struct FSampledSequenceDrawingParams
 	{
-		double MaxDisplayedValue = 1.f;
+		double MaxDisplayedValue = 1;
 		float DimensionSlotMargin = 2.f;
-		float MaxSequenceHeight = 0.9f;
-		float MinSequenceHeight = 0.1f;
+		float MaxSequenceHeightRatio = 0.95f;
+		float MinSequenceHeightRatio = 0.1f;
 		float MinScaledBinValue = 0.001f;
 		float VerticalZoomFactor = 1.f;
+		ESampledSequenceDrawOrientation Orientation = ESampledSequenceDrawOrientation::Horizontal;
 	};
 
 	/**
-		* Represents an horizontal slot in which samples or sample bins can be drawn
+	* Represents a dimensional slot in which samples, sample bins or a grid can be drawn
 	*/
-	struct FHorizontalDimensionSlot
+	struct FDimensionSlot
 	{
-		explicit FHorizontalDimensionSlot(const uint16 DimensionToDraw, const uint16 TotalNumDimensions, const FGeometry& InAllottedGeometry);
+		explicit FDimensionSlot(const uint16 DimensionToDraw, const uint16 TotalNumDimensions, const FGeometry& InAllottedGeometry, const FSampledSequenceDrawingParams& Params);
 
 		float Top;
 		float Center;
@@ -124,10 +134,20 @@ namespace SampledSequenceDrawingUtils
 		float Height;
 	};
 
-	struct FSampleBinCoordinates
+	struct F2DLineCoordinates
+	{
+		FVector2D A;
+		FVector2D B;
+	};
+
+	/**
+	 * DrawCoordinates: Coordinates of different grid lines
+	 * PositionRatios: the Ratio of each grid line position in relationship to the main axis length
+	 */
+	struct FGridData
 	{
-		FVector2D Top;
-		FVector2D Bottom;
+		TArray<F2DLineCoordinates> DrawCoordinates;
+		TArray<float> PositionRatios;
 	};
 
 	/**
@@ -140,7 +160,7 @@ namespace SampledSequenceDrawingUtils
 		* @param NDimensions				Number of interleaved dimensions in the time series
 		* @param Params						Drawing Params
 	*/
-	void GenerateSampleBinsCoordinatesForGeometry(TArray<FSampleBinCoordinates>& OutDrawCoordinates, const FGeometry& InAllottedGeometry, const TArray<TRange<float>>& InSampleBins, const uint16 NDimensions, const FSampledSequenceDrawingParams Params = FSampledSequenceDrawingParams());
+	void GenerateSampleBinsCoordinatesForGeometry(TArray<F2DLineCoordinates>& OutDrawCoordinates, const FGeometry& InAllottedGeometry, const TArray<TRange<float>>& InSampleBins, const uint16 NDimensions, const FSampledSequenceDrawingParams Params = FSampledSequenceDrawingParams());
 
 	/**
 	* Generates an array of coordinates to draw single samples in succession horizontally.
@@ -154,4 +174,30 @@ namespace SampledSequenceDrawingUtils
 	*/
 	void GenerateSequencedSamplesCoordinatesForGeometry(TArray<FVector2D>& OutDrawCoordinates, TArrayView<const float> InSampleData, const FGeometry& InAllottedGeometry, const uint16 NDimensions, const FFixedSampledSequenceGridMetrics InGridMetrics, const FSampledSequenceDrawingParams Params = FSampledSequenceDrawingParams());
 
+	/**
+	* Generates an evenly divided grid from a given geometry, following the orientation provided in the Params struct argument.
+	*
+	* @param OutGridData			The generated grid coordinates and ratios per each dimension
+	* @param InAllottedGeometry		The geometry the grid will be drawn into
+	* @param NDimensions			Dimensions to account for in the given geometry
+	* @param NumGridDivisions		Number of divisions the grid should be split into 
+	* @param Params					Drawing Params
+	*/
+	void GenerateEvenlySplitGridForGeometry(TArray<FGridData>& OutGridData, const FGeometry& InAllottedGeometry, const uint16 NDimensions, const uint32 NumGridDivisions /*= 5*/, const FSampledSequenceDrawingParams Params = FSampledSequenceDrawingParams());
+	
+	void AUDIOWIDGETS_API GenerateEvenlySplitGridForLine(TArray<double>& OutDrawCoordinates, TArray<float>& OutLinePositionRatios, const float LineLength, const uint32 NumGridDivisions /*= 2*/, bool bEmptyOutArrays = true);
+
+	/**
+	* Generates a midline divided grid from a given geometry, following the orientation provided in the Params struct argument.
+	* The grid space will be divided in half, with the resulting halves being divided again and so on until the desired Division Depth is reached. 
+	*
+	* @param OutGridData			The generated grid coordinates and ratios per each dimension
+	* @param InAllottedGeometry		The geometry the grid will be drawn into
+	* @param NDimensions			Dimensions to account for in the given geometry
+	* @param DivisionDepth			Number of divisions that should be run
+	* @param Params					Drawing Params
+	*/
+	void GenerateMidpointSplitGridForGeometry(TArray<FGridData>& OutGridData, const FGeometry& InAllottedGeometry, const uint16 NDimensions, const uint32 DivisionDepth = 3, const FSampledSequenceDrawingParams Params = FSampledSequenceDrawingParams());
+	
+	void AUDIOWIDGETS_API GenerateMidpointSplitGridForLine(TArray<double>& OutDrawCoordinates, TArray<float>& OutLinePositionRatios, const float LineLength, const uint32 DivisionDepth = 3, bool bEmptyOutArrays = true);
 }
\ No newline at end of file
