[AudioWidgets] Generic Fixed Sampled Sequence Viewer Widget

diff --git a/Engine/Plugins/Runtime/AudioWidgets/Source/AudioWidgets/Private/AudioWidgetsSlateTypes.cpp b/Engine/Plugins/Runtime/AudioWidgets/Source/AudioWidgets/Private/AudioWidgetsSlateTypes.cpp
--- a/Engine/Plugins/Runtime/AudioWidgets/Source/AudioWidgets/Private/AudioWidgetsSlateTypes.cpp
+++ b/Engine/Plugins/Runtime/AudioWidgets/Source/AudioWidgets/Private/AudioWidgetsSlateTypes.cpp
@@ -159,7 +159,7 @@ void FFixedSampleSequenceRulerStyle::GetResources(TArray< const FSlateBrush* >&
 FSampledSequenceValueGridOverlayStyle::FSampledSequenceValueGridOverlayStyle()
 	: GridColor(FLinearColor::Black)
 	, GridThickness(1.f)
-	, LabelTextColor(FLinearColor::Black)
+	, LabelTextColor(FLinearColor::White)
 	, LabelTextFont(FAppStyle::GetFontStyle("Regular"))
 	, DesiredWidth(AudioWidgetStylesSharedParams::DefaultWidth)
 	, DesiredHeight(AudioWidgetStylesSharedParams::DefaultHeight)
diff --git a/Engine/Plugins/Runtime/AudioWidgets/Source/AudioWidgets/Private/AudioWidgetsStyle.cpp b/Engine/Plugins/Runtime/AudioWidgets/Source/AudioWidgets/Private/AudioWidgetsStyle.cpp
--- a/Engine/Plugins/Runtime/AudioWidgets/Source/AudioWidgets/Private/AudioWidgetsStyle.cpp
+++ b/Engine/Plugins/Runtime/AudioWidgets/Source/AudioWidgets/Private/AudioWidgetsStyle.cpp
@@ -97,6 +97,28 @@ FAudioWidgetsStyle::FAudioWidgetsStyle()
 		.SetLabelPadding(3.0f)
 		.SetDefaultSliderRadius(50.0f));
 
+	/**
+	* SampledSequenceViewer Style
+	*/
+	Set("SampledSequenceViewer.Style", FSampledSequenceViewerStyle());
+
+	/**
+	* FixedSampledSequenceRuler Style
+	*/
+	Set("FixedSampledSequenceRuler.Style", FFixedSampleSequenceRulerStyle());
+
+	/**
+	* Playhead Overlay Style
+	*/
+	Set("PlayheadOverlay.Style", FPlayheadOverlayStyle());
+
+
+	/**
+	* ValueGrid Style
+	*/
+	Set("ValueGridOverlay.Style", FSampledSequenceValueGridOverlayStyle());
+
+
 	FSlateStyleRegistry::RegisterSlateStyle(*this);
 }
 
diff --git a/Engine/Plugins/Runtime/AudioWidgets/Source/AudioWidgets/Private/SFixedSampledSequenceViewer.cpp b/Engine/Plugins/Runtime/AudioWidgets/Source/AudioWidgets/Private/SFixedSampledSequenceViewer.cpp
--- a/Engine/Plugins/Runtime/AudioWidgets/Source/AudioWidgets/Private/SFixedSampledSequenceViewer.cpp
+++ b/Engine/Plugins/Runtime/AudioWidgets/Source/AudioWidgets/Private/SFixedSampledSequenceViewer.cpp
@@ -10,6 +10,7 @@ void SFixedSampledSequenceViewer::Construct(const FArguments& InArgs, TArrayView
 	Style = InArgs._Style;
 	DrawingParams = InArgs._SequenceDrawingParams;
 	bHideBackground = InArgs._HideBackground;
+	bHideGrid = InArgs._HideGrid;
 
 	check(Style);
 	SequenceColor = Style->SequenceColor;
@@ -54,8 +55,11 @@ int32 SFixedSampledSequenceViewer::OnPaint(const FPaintArgs& Args, const FGeomet
 
 	if (PixelWidth > 0)
 	{
-		DrawGridLines(AllottedGeometry, OutDrawElements, LayerId);
-
+		if (!bHideGrid)
+		{
+			DrawGridLines(AllottedGeometry, OutDrawElements, LayerId);
+		}
+		
 		if (SequenceDrawMode == ESequenceDrawMode::BinnedPeaks)
 		{
 			TArray<FVector2D> BinDrawPoints;
diff --git a/Engine/Plugins/Runtime/AudioWidgets/Source/AudioWidgets/Private/SFixedSampledSequenceViewerPanel.cpp b/Engine/Plugins/Runtime/AudioWidgets/Source/AudioWidgets/Private/SFixedSampledSequenceViewerPanel.cpp
new file mode 100644
--- /dev/null
+++ b/Engine/Plugins/Runtime/AudioWidgets/Source/AudioWidgets/Private/SFixedSampledSequenceViewerPanel.cpp
@@ -0,0 +1,211 @@
+// Copyright Epic Games, Inc. All Rights Reserved.
+
+#include "SFixedSampledSequenceViewerPanel.h"
+
+#include "FixedSampledSequenceGridData.h"
+#include "SFixedSampledSequenceRuler.h"
+#include "SFixedSampledSequenceViewer.h"
+#include "SparseSampledSequenceTransportCoordinator.h"
+#include "SPlayheadOverlay.h"
+#include "Widgets/Layout/SBorder.h"
+#include "Widgets/SBoxPanel.h"
+#include "Widgets/SCompoundWidget.h"
+#include "Widgets/SOverlay.h"
+
+
+void SFixedSampledSequenceViewerPanel::Construct(const FArguments& InArgs, const FFixedSampledSequenceView& InData)
+{
+	DataView = InData;
+	
+	DrawingParams = InArgs._SequenceDrawingParams;
+	SamplesRulerDisplayUnit = InArgs._SamplesRulerDisplayUnit;
+	bHidePlayhead = InArgs._HidePlayhead.Get();
+	bHideSamplesRuler = InArgs._HideSamplesRuler.Get();
+	bHideSamplesGrid = InArgs._HideSamplesGrid.Get();
+	bHideValueGrid = InArgs._HideValueGrid.Get();
+
+	CreateBackground(InArgs._SequenceViewerStyle);
+	CreateGridData(InArgs._SamplesRulerStyle);
+	CreateSequenceViewer(SamplesGridData.ToSharedRef(), DataView, InArgs._SequenceViewerStyle);
+	CreatePlayheadOverlay(InArgs._PlayheadOverlayStyle);
+	CreateSamplesRuler(SamplesGridData.ToSharedRef(), InArgs._SamplesRulerStyle);
+	CreateValueGridOverlay(InArgs._ValueGridMaxDivisionParameter, InArgs._ValueGridDivideMode, InArgs._ValueGridLabelGenerator, InArgs._ValueGridStyle);
+
+	CreateLayout();
+}
+
+void SFixedSampledSequenceViewerPanel::CreateLayout()
+{
+	check(SamplesRuler);
+	check(SequenceViewer);
+	check(PlayheadOverlay);
+	check(ValueGridOverlay);
+	check(BackgroundBorder);
+
+	TSharedPtr<SVerticalBox> MainBox = SNew(SVerticalBox);
+
+	if (!bHideSamplesRuler)
+	{
+		MainBox->AddSlot().AutoHeight()
+		[
+			SamplesRuler.ToSharedRef()
+		];
+	}
+
+	TSharedPtr<SOverlay> SequenceView = SNew(SOverlay);
+
+	SequenceView->AddSlot()
+	[
+		BackgroundBorder.ToSharedRef()
+	];
+
+	if (!bHideValueGrid)
+	{
+		SequenceView->AddSlot()
+		[
+			ValueGridOverlay.ToSharedRef()
+		];
+	}
+
+	SequenceView->AddSlot()
+	[
+		SequenceViewer.ToSharedRef()
+	];
+
+	if (!bHidePlayhead)
+	{
+		SequenceView->AddSlot()
+		[
+			PlayheadOverlay.ToSharedRef()
+		];
+	}
+
+	MainBox->AddSlot()
+	[
+		SequenceView.ToSharedRef()
+	];
+	
+
+	ChildSlot
+	[
+		MainBox.ToSharedRef()
+	];
+}
+
+void SFixedSampledSequenceViewerPanel::CreateSamplesRuler(TSharedRef<FFixedSampledSequenceGridData> InGridData, const FFixedSampleSequenceRulerStyle* RulerStyle)
+{
+	check(RulerStyle)
+	SamplesRuler = SNew(SFixedSampledSequenceRuler, InGridData).DisplayUnit(SamplesRulerDisplayUnit).Style(RulerStyle).DisplayPlayhead(!bHidePlayhead);
+	SamplesGridData->OnGridMetricsUpdated.AddSP(SamplesRuler.ToSharedRef(), &SFixedSampledSequenceRuler::UpdateGridMetrics);
+	SamplesRuler->OnTimeUnitMenuSelection.AddSP(this, &SFixedSampledSequenceViewerPanel::UpdateSamplesRulerDisplayUnit);
+}
+
+
+void SFixedSampledSequenceViewerPanel::CreatePlayheadOverlay(const FPlayheadOverlayStyle* PlayheadOverlayStyle)
+{
+	check(PlayheadOverlayStyle)
+	PlayheadOverlay = SNew(SPlayheadOverlay).Style(PlayheadOverlayStyle);
+}
+
+void SFixedSampledSequenceViewerPanel::CreateValueGridOverlay(const uint32 MaxDivisionParameter, const SampledSequenceValueGridOverlay::EGridDivideMode DivideMode, const TFunction<FText(const double)>& LabelGenerator, const FSampledSequenceValueGridOverlayStyle* ValueGridStyle)
+{
+	check(ValueGridStyle)
+
+	ValueGridOverlay = SNew(SSampledSequenceValueGridOverlay)
+	.MaxDivisionParameter(MaxDivisionParameter)
+	.DivideMode(DivideMode)
+	.ValueGridLabelGenerator(LabelGenerator)
+	.Style(ValueGridStyle)
+	.SequenceDrawingParams(DrawingParams)
+	.NumDimensions(DataView.NumDimensions);
+	
+}
+
+void SFixedSampledSequenceViewerPanel::CreateSequenceViewer(TSharedRef<FFixedSampledSequenceGridData> InGridData, const FFixedSampledSequenceView& InData, const FSampledSequenceViewerStyle* ViewerStyle)
+{
+	check(ViewerStyle)
+
+	SequenceViewer = SNew(SFixedSampledSequenceViewer, InData.SampleData, InData.NumDimensions, InGridData)
+		.Style(ViewerStyle)
+		.SequenceDrawingParams(DrawingParams)
+		.HideBackground(true)
+		.HideGrid(bHideSamplesGrid);
+
+	SamplesGridData->OnGridMetricsUpdated.AddSP(SequenceViewer.ToSharedRef(), &SFixedSampledSequenceViewer::UpdateGridMetrics);
+}
+
+void SFixedSampledSequenceViewerPanel::CreateBackground(const FSampledSequenceViewerStyle* ViewerStyle)
+{
+	check(ViewerStyle)
+
+	BackgroundBorder = SNew(SBorder)
+		.BorderImage(&ViewerStyle->BackgroundBrush)
+		.BorderBackgroundColor(ViewerStyle->SequenceBackgroundColor);
+}
+
+void SFixedSampledSequenceViewerPanel::CreateGridData(const FFixedSampleSequenceRulerStyle* RulerStyle)
+{
+	check(RulerStyle)
+
+	SamplesGridData = MakeShared<FFixedSampledSequenceGridData>(DataView.SampleData.Num() / DataView.NumDimensions, DataView.SampleRate, RulerStyle->DesiredWidth, &RulerStyle->TicksTextFont);
+}
+
+void SFixedSampledSequenceViewerPanel::ReceiveSequenceView(const FFixedSampledSequenceView InData, const uint32 FirstSampleIndex)
+{
+	if (SamplesGridData)
+	{
+		const uint32 FirstRenderedFrame = FirstSampleIndex / InData.NumDimensions;
+		const uint32 NumFrames = InData.SampleData.Num() / InData.NumDimensions;
+		SamplesGridData->UpdateDisplayRange(TRange<uint32>(FirstRenderedFrame, FirstRenderedFrame + NumFrames));
+	}
+
+	if (SequenceViewer)
+	{
+		SequenceViewer->UpdateView(InData.SampleData, InData.NumDimensions);
+	}
+}
+
+void SFixedSampledSequenceViewerPanel::SetPlayheadRatio(const float InRatio)
+{
+	CachedPlayheadRatio = InRatio;
+}
+
+void SFixedSampledSequenceViewerPanel::UpdateSamplesRulerDisplayUnit(const ESampledSequenceDisplayUnit InDisplayUnit)
+{
+	SamplesRulerDisplayUnit = InDisplayUnit;
+	SamplesRuler->UpdateDisplayUnit(SamplesRulerDisplayUnit);
+}
+
+void SFixedSampledSequenceViewerPanel::Tick(const FGeometry& AllottedGeometry, const double InCurrentTime, const float InDeltaTime)
+{
+	const float PaintedWidth = AllottedGeometry.GetLocalSize().X;
+
+	if (PaintedWidth != CachedPixelWidth)
+	{
+		CachedPixelWidth = PaintedWidth;
+
+		if (SamplesGridData)
+		{
+			SamplesGridData->UpdateGridMetrics(PaintedWidth);
+		}
+	}
+
+	UpdatePlayheadPosition(PaintedWidth);
+}
+
+void SFixedSampledSequenceViewerPanel::UpdatePlayheadPosition(const float PaintedWidth)
+{
+	float PlayheadX = CachedPlayheadRatio * PaintedWidth;
+	PlayheadX = SamplesGridData ? SamplesGridData->SnapPositionToClosestFrame(PlayheadX) : PlayheadX;
+
+
+	if (!bHidePlayhead && PlayheadOverlay)
+	{
+		PlayheadOverlay->SetPlayheadPosition(PlayheadX);
+	}
+
+	if (!bHideSamplesRuler && SamplesRuler)
+	{
+		SamplesRuler->SetPlayheadPosition(PlayheadX);
+	}
+}
diff --git a/Engine/Plugins/Runtime/AudioWidgets/Source/AudioWidgets/Private/SSampledSequenceValueGridOverlay.cpp b/Engine/Plugins/Runtime/AudioWidgets/Source/AudioWidgets/Private/SSampledSequenceValueGridOverlay.cpp
--- a/Engine/Plugins/Runtime/AudioWidgets/Source/AudioWidgets/Private/SSampledSequenceValueGridOverlay.cpp
+++ b/Engine/Plugins/Runtime/AudioWidgets/Source/AudioWidgets/Private/SSampledSequenceValueGridOverlay.cpp
@@ -5,13 +5,22 @@
 #include "Fonts/FontMeasure.h"
 #include "Framework/Application/SlateApplication.h"
 
+namespace SampledSequenceValueGridOverlayPrivate
+{
+	TFunction<FText(const double)> DefaultLabelGenerator = [](const double GridLineValue)
+	{
+		FNumberFormattingOptions DefaultFormatting;
+		return FText::AsNumber(GridLineValue, &DefaultFormatting);
+	};
+}
+
 void SSampledSequenceValueGridOverlay::Construct(const FArguments& InArgs)
 {
 	check(InArgs._Style);
 
 	GridSplitMode = InArgs._DivideMode;
 	MaxDivisionParameter = InArgs._MaxDivisionParameter;
-	OnValueGridLabel = InArgs._ValueGridLabelGenerator;
+	OnValueGridLabel = InArgs._ValueGridLabelGenerator ? InArgs._ValueGridLabelGenerator : SampledSequenceValueGridOverlayPrivate::DefaultLabelGenerator;
 	NumDimensions = InArgs._NumDimensions;
 	DrawingParams = InArgs._SequenceDrawingParams;
 	bHideLabels = InArgs._HideLabels;
diff --git a/Engine/Plugins/Runtime/AudioWidgets/Source/AudioWidgets/Public/AudioWidgetsStyle.h b/Engine/Plugins/Runtime/AudioWidgets/Source/AudioWidgets/Public/AudioWidgetsStyle.h
--- a/Engine/Plugins/Runtime/AudioWidgets/Source/AudioWidgets/Public/AudioWidgetsStyle.h
+++ b/Engine/Plugins/Runtime/AudioWidgets/Source/AudioWidgets/Public/AudioWidgetsStyle.h
@@ -7,7 +7,7 @@
 /**
  * Slate style set that defines all the styles for audio widgets
  */
-class FAudioWidgetsStyle
+class AUDIOWIDGETS_API FAudioWidgetsStyle
 	: public FSlateStyleSet
 
 {
diff --git a/Engine/Plugins/Runtime/AudioWidgets/Source/AudioWidgets/Public/SFixedSampledSequenceViewer.h b/Engine/Plugins/Runtime/AudioWidgets/Source/AudioWidgets/Public/SFixedSampledSequenceViewer.h
--- a/Engine/Plugins/Runtime/AudioWidgets/Source/AudioWidgets/Public/SFixedSampledSequenceViewer.h
+++ b/Engine/Plugins/Runtime/AudioWidgets/Source/AudioWidgets/Public/SFixedSampledSequenceViewer.h
@@ -20,6 +20,7 @@ public:
 
 	SLATE_BEGIN_ARGS(SFixedSampledSequenceViewer) 
 		: _HideBackground(false)
+		, _HideGrid(false)
 	{
 	}
 	
@@ -27,6 +28,8 @@ public:
 
 	SLATE_ARGUMENT(bool, HideBackground)
 
+	SLATE_ARGUMENT(bool, HideGrid)
+
 	SLATE_STYLE_ARGUMENT(FSampledSequenceViewerStyle, Style)
 		
 	SLATE_END_ARGS()
@@ -66,7 +69,8 @@ private:
 	float DesiredWidth = 0.f;
 	float SampleMarkersSize = 2.5f;
 	float SequenceLineThickness = 1.f;
-	bool bHideBackground;
+	bool bHideBackground = false;
+	bool bHideGrid = false;
 
 	uint32 CachedPixelWidth = 0; 
 	float CachedPixelHeight = 0.f;
diff --git a/Engine/Plugins/Runtime/AudioWidgets/Source/AudioWidgets/Public/SFixedSampledSequenceViewerPanel.h b/Engine/Plugins/Runtime/AudioWidgets/Source/AudioWidgets/Public/SFixedSampledSequenceViewerPanel.h
new file mode 100644
--- /dev/null
+++ b/Engine/Plugins/Runtime/AudioWidgets/Source/AudioWidgets/Public/SFixedSampledSequenceViewerPanel.h
@@ -0,0 +1,125 @@
+// Copyright Epic Games, Inc. All Rights Reserved.
+
+#pragma once
+
+#include "AudioWidgetsSlateTypes.h"
+#include "AudioWidgetsStyle.h"
+#include "IFixedSampledSequenceViewReceiver.h"
+#include "SampledSequenceDisplayUnit.h"
+#include "SampledSequenceDrawingUtils.h"
+#include "SSampledSequenceValueGridOverlay.h"
+#include "Widgets/SCompoundWidget.h"
+
+class FFixedSampledSequenceGridData;
+class FSparseSampledSequenceTransportCoordinator;
+class FWaveformEditorGridData;
+class SBorder;
+class SFixedSampledSequenceRuler;
+class SFixedSampledSequenceViewer;
+class SPlayheadOverlay;
+
+class AUDIOWIDGETS_API SFixedSampledSequenceViewerPanel : public SCompoundWidget, public IFixedSampledSequenceViewReceiver
+{
+public:
+
+	SLATE_BEGIN_ARGS(SFixedSampledSequenceViewerPanel) 
+		: _HidePlayhead(false)
+		, _HideSamplesGrid(false)
+		, _HideSamplesRuler(false)
+		, _HideValueGrid(false)
+		, _ValueGridDivideMode(SampledSequenceValueGridOverlay::EGridDivideMode::MidSplit)
+		, _ValueGridMaxDivisionParameter(2)
+		, _SequenceDrawingParams(SampledSequenceDrawingUtils::FSampledSequenceDrawingParams())
+		, _SamplesRulerDisplayUnit(ESampledSequenceDisplayUnit::Seconds)
+		, _SequenceViewerStyle(&FAudioWidgetsStyle::Get().GetWidgetStyle< FSampledSequenceViewerStyle >("SampledSequenceViewer.Style"))
+		, _SamplesRulerStyle(&FAudioWidgetsStyle::Get().GetWidgetStyle< FFixedSampleSequenceRulerStyle >("FixedSampledSequenceRuler.Style"))
+		, _PlayheadOverlayStyle(&FAudioWidgetsStyle::Get().GetWidgetStyle< FPlayheadOverlayStyle >("PlayheadOverlay.Style"))
+		, _ValueGridStyle(&FAudioWidgetsStyle::Get().GetWidgetStyle< FSampledSequenceValueGridOverlayStyle >("ValueGridOverlay.Style"))
+	{}
+		
+		SLATE_DEFAULT_SLOT(FArguments, InArgs)
+
+		/** Whether the playhead should be drawn or not */
+		SLATE_ATTRIBUTE(bool, HidePlayhead)
+
+		/** Whether the samples grid should be drawn or not */
+		SLATE_ATTRIBUTE(bool, HideSamplesGrid)
+
+		/** Whether the samples ruler should be drawn or not */
+		SLATE_ATTRIBUTE(bool, HideSamplesRuler)
+
+		/** Whether the value grid should be drawn or not */
+		SLATE_ATTRIBUTE(bool, HideValueGrid)
+
+		/** Regulates how the value grid is divided */
+		SLATE_ARGUMENT(SampledSequenceValueGridOverlay::EGridDivideMode, ValueGridDivideMode)
+
+		/** Maximum number of divisions in the value grid */
+		SLATE_ARGUMENT(uint32, ValueGridMaxDivisionParameter)
+
+		/** Allows customization of labels displayed in the value grid */
+		SLATE_ARGUMENT(TFunction<FText(const double)>, ValueGridLabelGenerator)
+
+		/** Sequence drawing params */
+		SLATE_ARGUMENT(SampledSequenceDrawingUtils::FSampledSequenceDrawingParams, SequenceDrawingParams)
+
+		/** Samples ruler display unit */
+		SLATE_ARGUMENT(ESampledSequenceDisplayUnit, SamplesRulerDisplayUnit)
+
+		/** Sequence viewer widget style */
+		SLATE_STYLE_ARGUMENT(FSampledSequenceViewerStyle, SequenceViewerStyle)
+
+		/** Time Ruler widget style */
+		SLATE_STYLE_ARGUMENT(FFixedSampleSequenceRulerStyle, SamplesRulerStyle)
+
+		/** Playhead overlay widget style */
+		SLATE_STYLE_ARGUMENT(FPlayheadOverlayStyle, PlayheadOverlayStyle)
+
+		/** Value grid widget style */
+		SLATE_STYLE_ARGUMENT(FSampledSequenceValueGridOverlayStyle, ValueGridStyle)
+
+
+	SLATE_END_ARGS()
+
+	void Construct(const FArguments& InArgs, const FFixedSampledSequenceView& InData);
+	virtual void Tick(const FGeometry& AllottedGeometry, const double InCurrentTime, const float InDeltaTime) override;
+
+	/*IFixedSampledSequenceViewReceiver Interface*/
+	virtual void ReceiveSequenceView(const FFixedSampledSequenceView InData, const uint32 FirstSampleIndex) override;
+	
+	void SetPlayheadRatio(const float InRatio);
+	
+private:
+	virtual void CreateLayout();
+
+	TSharedPtr<FFixedSampledSequenceGridData> SamplesGridData;
+	TSharedPtr<SBorder> BackgroundBorder;
+	TSharedPtr<SFixedSampledSequenceRuler> SamplesRuler;
+	TSharedPtr<SFixedSampledSequenceViewer> SequenceViewer;
+	TSharedPtr<SPlayheadOverlay> PlayheadOverlay;
+	TSharedPtr<SSampledSequenceValueGridOverlay> ValueGridOverlay;
+
+
+	void CreateGridData(const FFixedSampleSequenceRulerStyle* RulerStyle);
+	void CreatePlayheadOverlay(const FPlayheadOverlayStyle* PlayheadOverlayStyle);
+	void CreateValueGridOverlay(const uint32 MaxDivisionParameter, const SampledSequenceValueGridOverlay::EGridDivideMode DivideMode, const TFunction<FText(const double)>& LabelGenerator, const FSampledSequenceValueGridOverlayStyle* ValueGridStyle);
+	void CreateSamplesRuler(TSharedRef<FFixedSampledSequenceGridData> InGridData, const FFixedSampleSequenceRulerStyle* RulerStyle);
+	void CreateSequenceViewer(TSharedRef<FFixedSampledSequenceGridData> InGridData, const FFixedSampledSequenceView& InData, const FSampledSequenceViewerStyle* ViewerStyle);
+	void CreateBackground(const FSampledSequenceViewerStyle* ViewerStyle);
+	
+	void UpdateSamplesRulerDisplayUnit(const ESampledSequenceDisplayUnit InDisplayUnit);
+	void UpdatePlayheadPosition(const float PaintedWidth);
+
+
+	float CachedPixelWidth = 0.f;
+
+	ESampledSequenceDisplayUnit SamplesRulerDisplayUnit;
+	FFixedSampledSequenceView DataView;
+	SampledSequenceDrawingUtils::FSampledSequenceDrawingParams DrawingParams;
+
+	float CachedPlayheadRatio = 0.f;
+	bool bHidePlayhead = false;
+	bool bHideSamplesGrid = false;
+	bool bHideSamplesRuler = false;
+	bool bHideValueGrid = false;
+};
\ No newline at end of file
diff --git a/Engine/Plugins/Runtime/AudioWidgets/Source/AudioWidgets/Public/SSampledSequenceValueGridOverlay.h b/Engine/Plugins/Runtime/AudioWidgets/Source/AudioWidgets/Public/SSampledSequenceValueGridOverlay.h
--- a/Engine/Plugins/Runtime/AudioWidgets/Source/AudioWidgets/Public/SSampledSequenceValueGridOverlay.h
+++ b/Engine/Plugins/Runtime/AudioWidgets/Source/AudioWidgets/Public/SSampledSequenceValueGridOverlay.h
@@ -16,12 +16,6 @@ namespace SampledSequenceValueGridOverlay
 		MidSplit,
 		COUNT
 	};
-
-	TFunction<FText(const double)> DefaultLabelGenerator = [](const double GridLineValue)
-	{
-		FNumberFormattingOptions DefaultFormatting;
-		return FText::AsNumber(GridLineValue, &DefaultFormatting);
-	};
 }
 
 class AUDIOWIDGETS_API SSampledSequenceValueGridOverlay : public SLeafWidget
@@ -33,7 +27,6 @@ public:
 		, _HideLabels(false)
 		, _HideGrid(false)
 		, _MaxDivisionParameter(2)
-		, _ValueGridLabelGenerator(SampledSequenceValueGridOverlay::DefaultLabelGenerator)
 	{
 	}
 
