[AudioWidgets] new scaling logic when generating waveform bins to take dimension slot margins into account
#trivial

diff --git a/Engine/Plugins/Runtime/AudioWidgets/Source/AudioWidgets/Private/SampledSequenceDrawingUtils.cpp b/Engine/Plugins/Runtime/AudioWidgets/Source/AudioWidgets/Private/SampledSequenceDrawingUtils.cpp
--- a/Engine/Plugins/Runtime/AudioWidgets/Source/AudioWidgets/Private/SampledSequenceDrawingUtils.cpp
+++ b/Engine/Plugins/Runtime/AudioWidgets/Source/AudioWidgets/Private/SampledSequenceDrawingUtils.cpp
@@ -49,7 +49,7 @@ SampledSequenceDrawingUtils::FDimensionSlot::FDimensionSlot(const uint16 Dimensi
 
 	const double FullSlotHeight = GeometryLength / TotalNumDimensions;
 	const double MarginedTop = (FullSlotHeight * DimensionToDraw) + Params.DimensionSlotMargin;
-	Height = FMath::Clamp(FullSlotHeight - (Params.DimensionSlotMargin * 2), FullSlotHeight * Params.MinSequenceHeightRatio, GeometryLength / TotalNumDimensions);
+	Height = FMath::Clamp(FullSlotHeight - (Params.DimensionSlotMargin * 2), FullSlotHeight * Params.MinSequenceHeightRatio, FullSlotHeight);
 	const double MarginedBottom = MarginedTop + Height;
 
 	Center = MarginedTop + ((MarginedBottom - MarginedTop) / 2.f);
@@ -68,8 +68,8 @@ void SampledSequenceDrawingUtils::GenerateSampleBinsCoordinatesForGeometry(TArra
 	const uint32 PixelWidth = FMath::FloorToInt(InAllottedGeometry.GetLocalSize().X);
 	check(PixelWidth * NDimensions == InSampleBins.Num());
 
-	const float HeightScale = InAllottedGeometry.GetLocalSize().Y / (2.f * Params.MaxDisplayedValue * NDimensions) * Params.MaxSequenceHeightRatio * Params.VerticalZoomFactor;
-
+	const FDimensionSlot SampleSlotForScale(0, NDimensions, InAllottedGeometry, Params);
+	const float HeightScale = SampleSlotForScale.Height / (2.f * Params.MaxDisplayedValue) * Params.VerticalZoomFactor;
 
 	OutDrawCoordinates.SetNumUninitialized(InSampleBins.Num());
 	F2DLineCoordinates* OutCoordinatesData = OutDrawCoordinates.GetData();
