[AudioWidgets] Minor ticks in some  audio widgets grids are explicitely constrained to geometry to avoid overflowing
#trivial

diff --git a/Engine/Plugins/Runtime/AudioWidgets/Source/AudioWidgets/Private/SFixedSampledSequenceRuler.cpp b/Engine/Plugins/Runtime/AudioWidgets/Source/AudioWidgets/Private/SFixedSampledSequenceRuler.cpp
--- a/Engine/Plugins/Runtime/AudioWidgets/Source/AudioWidgets/Private/SFixedSampledSequenceRuler.cpp
+++ b/Engine/Plugins/Runtime/AudioWidgets/Source/AudioWidgets/Private/SFixedSampledSequenceRuler.cpp
@@ -87,16 +87,21 @@ void SFixedSampledSequenceRuler::DrawRulerTicks(const FGeometry& AllottedGeometr
 		for (int32 MinorTickIndex = 1; MinorTickIndex < GridMetrics.NumMinorGridDivisions; ++MinorTickIndex)
 		{
 			const double MinorTickX = MajorTickX + MinorGridXStep * MinorTickIndex;
-			LinePoints[0] = FVector2D(MinorTickX, MinorTickY);
-			LinePoints[1] = FVector2D(MinorTickX, AllottedGeometry.Size.Y);
-			FSlateDrawElement::MakeLines(
-				OutDrawElements,
-				++LayerId,
-				AllottedGeometry.ToPaintGeometry(),
-				LinePoints,
-				ESlateDrawEffect::None,
-				TicksColor.GetSpecifiedColor(),
-				false);
+
+			if (MinorTickX < AllottedGeometry.Size.X)
+			{
+				LinePoints[0] = FVector2D(MinorTickX, MinorTickY);
+				LinePoints[1] = FVector2D(MinorTickX, AllottedGeometry.Size.Y);
+
+				FSlateDrawElement::MakeLines(
+					OutDrawElements,
+					++LayerId,
+					AllottedGeometry.ToPaintGeometry(),
+					LinePoints,
+					ESlateDrawEffect::None,
+					TicksColor.GetSpecifiedColor(),
+					false);
+			}
 		}
 		
 		uint32 TickFrame = FMath::RoundToInt32(MajorTickX / GridMetrics.PixelsPerFrame) + GridMetrics.StartFrame;
diff --git a/Engine/Plugins/Runtime/AudioWidgets/Source/AudioWidgets/Private/SFixedSampledSequenceViewer.cpp b/Engine/Plugins/Runtime/AudioWidgets/Source/AudioWidgets/Private/SFixedSampledSequenceViewer.cpp
--- a/Engine/Plugins/Runtime/AudioWidgets/Source/AudioWidgets/Private/SFixedSampledSequenceViewer.cpp
+++ b/Engine/Plugins/Runtime/AudioWidgets/Source/AudioWidgets/Private/SFixedSampledSequenceViewer.cpp
@@ -201,23 +201,24 @@ void SFixedSampledSequenceViewer::DrawGridLines(const FGeometry& AllottedGeometr
 			MajorGridLineColor.GetSpecifiedColor(),
 			false);
 
-
 		for (int32 MinorLineIndex = 1; MinorLineIndex < GridMetrics.NumMinorGridDivisions; ++MinorLineIndex)
 		{
 			const double MinorLineX = MajorLineX + MinorGridXStep * MinorLineIndex;
 
-			LinePoints[0] = FVector2D(MinorLineX, 0.);
-			LinePoints[1] = FVector2D(MinorLineX, AllottedGeometry.Size.Y);
-
-			FSlateDrawElement::MakeLines(
-				OutDrawElements,
-				++LayerId,
-				AllottedGeometry.ToPaintGeometry(),
-				LinePoints,
-				ESlateDrawEffect::None,
-				MinorGridLineColor.GetSpecifiedColor(),
-				false);
+			if (MinorLineX < AllottedGeometry.Size.X)
+			{
+				LinePoints[0] = FVector2D(MinorLineX, 0.);
+				LinePoints[1] = FVector2D(MinorLineX, AllottedGeometry.Size.Y);
 
+				FSlateDrawElement::MakeLines(
+					OutDrawElements,
+					++LayerId,
+					AllottedGeometry.ToPaintGeometry(),
+					LinePoints,
+					ESlateDrawEffect::None,
+					MinorGridLineColor.GetSpecifiedColor(),
+					false);
+			}
 		}
 	}
 
