[WaveformEditor] Fix sample bins reaching waveform center at certain zoom levels when they shouldn't + connect end of one bin to beginning of next one with a line to avoid scattered view

diff --git a/Engine/Plugins/Runtime/AudioWidgets/Source/AudioWidgets/Private/SFixedSampledSequenceViewer.cpp b/Engine/Plugins/Runtime/AudioWidgets/Source/AudioWidgets/Private/SFixedSampledSequenceViewer.cpp
--- a/Engine/Plugins/Runtime/AudioWidgets/Source/AudioWidgets/Private/SFixedSampledSequenceViewer.cpp
+++ b/Engine/Plugins/Runtime/AudioWidgets/Source/AudioWidgets/Private/SFixedSampledSequenceViewer.cpp
@@ -65,10 +65,10 @@ int32 SFixedSampledSequenceViewer::OnPaint(const FPaintArgs& Args, const FGeomet
 			TArray<FVector2D> BinDrawPoints;
 			BinDrawPoints.SetNumUninitialized(2);
 
-			for (const SampledSequenceDrawingUtils::F2DLineCoordinates& BinCoordinates : CachedBinsDrawCoordinates)
+			for (int32 i = 0; i < CachedBinsDrawCoordinates.Num(); ++i)
 			{
-				BinDrawPoints[0] = BinCoordinates.A;
-				BinDrawPoints[1] = BinCoordinates.B;
+				BinDrawPoints[0] = CachedBinsDrawCoordinates[i].A;
+				BinDrawPoints[1] = CachedBinsDrawCoordinates[i].B;
 
 				FSlateDrawElement::MakeLines(
 					OutDrawElements,
@@ -78,6 +78,24 @@ int32 SFixedSampledSequenceViewer::OnPaint(const FPaintArgs& Args, const FGeomet
 					ESlateDrawEffect::None,
 					SequenceColor.GetSpecifiedColor()
 				);
+
+				const int32 NextBinCoordinates = i + NumChannels;
+
+				if (NextBinCoordinates < CachedBinsDrawCoordinates.Num())
+				{
+					BinDrawPoints[0] = CachedBinsDrawCoordinates[i].B;
+					BinDrawPoints[1] = CachedBinsDrawCoordinates[NextBinCoordinates].A;
+
+					FSlateDrawElement::MakeLines(
+						OutDrawElements,
+						LayerId,
+						AllottedGeometry.ToPaintGeometry(),
+						BinDrawPoints,
+						ESlateDrawEffect::None,
+						SequenceColor.GetSpecifiedColor()
+					);
+
+				}
 			}
 
 		}
@@ -127,7 +145,7 @@ int32 SFixedSampledSequenceViewer::OnPaint(const FPaintArgs& Args, const FGeomet
 
 		++LayerId;
 	}
-
+	
 	return LayerId;
 }
 
@@ -270,4 +288,4 @@ void SFixedSampledSequenceViewer::OnStyleUpdated(const FNotifyingAudioWidgetStyl
 	SequenceLineThickness = Style->SequenceLineThickness;
 	DesiredWidth = Style->DesiredWidth;
 	DesiredHeight = Style->DesiredHeight;
-}
+}
\ No newline at end of file
diff --git a/Engine/Plugins/Runtime/AudioWidgets/Source/AudioWidgets/Private/SampledSequenceDrawingUtils.cpp b/Engine/Plugins/Runtime/AudioWidgets/Source/AudioWidgets/Private/SampledSequenceDrawingUtils.cpp
--- a/Engine/Plugins/Runtime/AudioWidgets/Source/AudioWidgets/Private/SampledSequenceDrawingUtils.cpp
+++ b/Engine/Plugins/Runtime/AudioWidgets/Source/AudioWidgets/Private/SampledSequenceDrawingUtils.cpp
@@ -83,11 +83,20 @@ void SampledSequenceDrawingUtils::GenerateSampleBinsCoordinatesForGeometry(TArra
 			uint32 PeakIndex = Pixel * NDimensions + Channel;
 			const TRange<float>& MinMaxBin = InSampleBins[PeakIndex];
 
-			const float SampleMaxScaled = MinMaxBin.GetUpperBoundValue() * HeightScale > Params.MinScaledBinValue ? MinMaxBin.GetUpperBoundValue() * HeightScale : Params.MinSequenceHeightRatio;
-			const float SampleMinScaled = MinMaxBin.GetLowerBoundValue() * HeightScale < -1 * Params.MinScaledBinValue ? MinMaxBin.GetLowerBoundValue() * HeightScale : -1.f * Params.MinSequenceHeightRatio;
+			const float SampleMaxScaled = MinMaxBin.GetUpperBoundValue() * HeightScale;
+			const float SampleMinScaled = MinMaxBin.GetLowerBoundValue() * HeightScale;
 
-			const float Top = FMath::Max(ChannelBoundaries.Center - SampleMaxScaled, ChannelBoundaries.Top);
-			const float Bottom = FMath::Min(ChannelBoundaries.Center - SampleMinScaled, ChannelBoundaries.Bottom);
+			float Top = FMath::Max(ChannelBoundaries.Center - SampleMaxScaled, ChannelBoundaries.Top);
+			float Bottom = FMath::Min(ChannelBoundaries.Center - SampleMinScaled, ChannelBoundaries.Bottom);
+
+			const bool bResizeBin = FMath::Abs(Top - Bottom) < Params.MinScaledBinValue;
+
+			if (bResizeBin)
+			{
+				const float MinHeightHalfSize = Params.MinScaledBinValue / 2.f;
+				Top = Top <= Bottom ? Top - MinHeightHalfSize : Top + MinHeightHalfSize;
+				Bottom = Top <= Bottom ? Bottom + MinHeightHalfSize : Top - MinHeightHalfSize;
+			}
 
 			OutCoordinatesData[PeakIndex].A = FVector2D(Pixel, Top);
 			OutCoordinatesData[PeakIndex].B = FVector2D(Pixel, Bottom);
@@ -233,4 +242,4 @@ void SampledSequenceDrawingUtils::GenerateEvenlySplitGridForLine(TArray<double>&
 
 	OutLinePositionRatios.Add(1.f);
 	OutDrawCoordinates.Add(LineLength);
-}
+}
\ No newline at end of file
diff --git a/Engine/Plugins/Runtime/AudioWidgets/Source/AudioWidgets/Public/SampledSequenceDrawingUtils.h b/Engine/Plugins/Runtime/AudioWidgets/Source/AudioWidgets/Public/SampledSequenceDrawingUtils.h
--- a/Engine/Plugins/Runtime/AudioWidgets/Source/AudioWidgets/Public/SampledSequenceDrawingUtils.h
+++ b/Engine/Plugins/Runtime/AudioWidgets/Source/AudioWidgets/Public/SampledSequenceDrawingUtils.h
@@ -43,7 +43,8 @@ namespace SampledSequenceDrawingUtils
 
 			for (uint16 Channel = 0; Channel < NDimensions; ++Channel)
 			{
-				SamplesType MaxSampleValue = TNumericLimits<SamplesType>::Min();
+				
+				SamplesType MaxSampleValue = TNumericLimits<SamplesType>::Lowest();
 				SamplesType MinSampleValue = TNumericLimits<SamplesType>::Max();
 
 				for (uint32 Frame = FirstBinnedFrame; Frame < LastBinnedFrame; ++Frame)
@@ -98,7 +99,7 @@ namespace SampledSequenceDrawingUtils
 		* @param DimensionSlotMargin				Margin to keep from the dimension slot boundaries (pixels)
 		* @param MinSequenceHeightRatio				The minimum height (ratio) the drawn sequence can take in a channel slot
 		* @param MaxSequenceHeightRatio				The maximum height (ratio) the drawn sequence can take in a channel slot
-		* @param MinScaledBinValue					Minimum value a scaled bin can have
+		* @param MinScaledBinValue					Minimum value (in pixels) a bin can have
 		* @param VerticalZoomFactor					Sequence zoom factor
 	*/
 
@@ -116,7 +117,7 @@ namespace SampledSequenceDrawingUtils
 		float DimensionSlotMargin = 2.f;
 		float MaxSequenceHeightRatio = 0.95f;
 		float MinSequenceHeightRatio = 0.1f;
-		float MinScaledBinValue = 0.001f;
+		float MinScaledBinValue = 0.1f;
 		float VerticalZoomFactor = 1.f;
 		ESampledSequenceDrawOrientation Orientation = ESampledSequenceDrawOrientation::Horizontal;
 	};
