[AudioWidgets] Remove dynamic style updates as advised by [REDACTED_PERSON]
- remove FNotifyingAudioWidgetStyle from AudioWidgets
- Styles are used as a payload that is sent as dynamic widgets are updated
- expose methods to update underlying widgets to SFixedSampledSequenceViewerPanel

diff --git a/Engine/Plugins/Editor/WaveformEditor/Source/WaveformEditorWidgets/Private/STransformedWaveformViewPanel.cpp b/Engine/Plugins/Editor/WaveformEditor/Source/WaveformEditorWidgets/Private/STransformedWaveformViewPanel.cpp
--- a/Engine/Plugins/Editor/WaveformEditor/Source/WaveformEditorWidgets/Private/STransformedWaveformViewPanel.cpp
+++ b/Engine/Plugins/Editor/WaveformEditor/Source/WaveformEditorWidgets/Private/STransformedWaveformViewPanel.cpp
@@ -2,7 +2,6 @@
 
  #include "STransformedWaveformViewPanel.h"
 
-#include "AudioWidgetsSlateTypes.h"
 #include "DSP/Dsp.h"
 #include "SampledSequenceDisplayUnit.h"
 #include "SFixedSampledSequenceRuler.h"
@@ -111,19 +110,17 @@ void STransformedWaveformViewPanel::CreateLayout()
 
 void STransformedWaveformViewPanel::SetUpTimeRuler(TSharedRef<FWaveformEditorGridData> InGridData)
 {
-	FFixedSampleSequenceRulerStyle* TimeRulerStyle = &WaveformEditorStyle->GetRegisteredWidgetStyle<FFixedSampleSequenceRulerStyle>("WaveformEditorRuler.Style").Get();
-	check(TimeRulerStyle);
+	const FFixedSampleSequenceRulerStyle& TimeRulerStyle = WaveformEditorStyle->GetWidgetStyle<FFixedSampleSequenceRulerStyle>("WaveformEditorRuler.Style");
 
-	TimeRuler = SNew(SFixedSampledSequenceRuler, InGridData).DisplayUnit(DisplayUnit).Style(TimeRulerStyle);
+	TimeRuler = SNew(SFixedSampledSequenceRuler, InGridData).DisplayUnit(DisplayUnit).Style(&TimeRulerStyle);
 	GridData->OnGridMetricsUpdated.AddSP(TimeRuler.ToSharedRef(), &SFixedSampledSequenceRuler::UpdateGridMetrics);
-	TimeRulerStyle->OnStyleUpdated.AddSP(TimeRuler.ToSharedRef(), &SFixedSampledSequenceRuler::OnStyleUpdated);
+	WaveformEditorStyle->OnNewTimeRulerStyle.AddSP(TimeRuler.ToSharedRef(), &SFixedSampledSequenceRuler::OnStyleUpdated);
 	TimeRuler->OnTimeUnitMenuSelection.AddSP(this, &STransformedWaveformViewPanel::UpdateDisplayUnit);
 }
 
 void STransformedWaveformViewPanel::SetUpInputRoutingOverlay()
 {
-	FSampledSequenceViewerStyle* ViewerStyle = &WaveformEditorStyle->GetRegisteredWidgetStyle<FSampledSequenceViewerStyle>("WaveformViewer.Style").Get();
-	check(ViewerStyle);
+	const FSampledSequenceViewerStyle& ViewerStyle = WaveformEditorStyle->GetWidgetStyle<FSampledSequenceViewerStyle>("WaveformViewer.Style");
 	
 	TArray<TSharedPtr<SWidget>> OverlaidWidgets;
 
@@ -134,35 +131,31 @@ void STransformedWaveformViewPanel::SetUpInputRoutingOverlay()
 
 	check(PlayheadOverlay)
 	OverlaidWidgets.Add(PlayheadOverlay);
-	InputRoutingOverlay = SNew(SWaveformEditorInputRoutingOverlay, OverlaidWidgets).Style(ViewerStyle);
+	InputRoutingOverlay = SNew(SWaveformEditorInputRoutingOverlay, OverlaidWidgets).Style(&ViewerStyle);
 }
 
 void STransformedWaveformViewPanel::SetupPlayheadOverlay()
 {
-	FPlayheadOverlayStyle* PlayheadOverlayStyle = &WaveformEditorStyle->GetRegisteredWidgetStyle<FPlayheadOverlayStyle>("WaveformEditorPlayheadOverlay.Style").Get();
-	check(PlayheadOverlayStyle);
-
-	PlayheadOverlay = SNew(SPlayheadOverlay).Style(PlayheadOverlayStyle);
-	PlayheadOverlayStyle->OnStyleUpdated.AddSP(PlayheadOverlay.ToSharedRef(), &SPlayheadOverlay::OnStyleUpdated);
+	const FPlayheadOverlayStyle& PlayheadOverlayStyle = WaveformEditorStyle->GetWidgetStyle<FPlayheadOverlayStyle>("WaveformEditorPlayheadOverlay.Style");
+	PlayheadOverlay = SNew(SPlayheadOverlay).Style(&PlayheadOverlayStyle);
+	WaveformEditorStyle->OnNewPlayheadOverlayStyle.AddSP(PlayheadOverlay.ToSharedRef(), &SPlayheadOverlay::OnStyleUpdated);
 }
 
 void STransformedWaveformViewPanel::SetUpWaveformViewer(TSharedRef<FWaveformEditorGridData> InGridData, const FFixedSampledSequenceView& InData)
 {
-	FSampledSequenceViewerStyle* WaveViewerStyle = &WaveformEditorStyle->GetRegisteredWidgetStyle<FSampledSequenceViewerStyle>("WaveformViewer.Style").Get();
-	check(WaveViewerStyle);
+	const FSampledSequenceViewerStyle& WaveViewerStyle = WaveformEditorStyle->GetWidgetStyle<FSampledSequenceViewerStyle>("WaveformViewer.Style");
 
-	WaveformViewer = SNew(SFixedSampledSequenceViewer, InData.SampleData, InData.NumDimensions, InGridData).Style(WaveViewerStyle).HideBackground(true);
+	WaveformViewer = SNew(SFixedSampledSequenceViewer, InData.SampleData, InData.NumDimensions, InGridData).Style(&WaveViewerStyle).HideBackground(true);
 		
-	WaveViewerStyle->OnStyleUpdated.AddSP(WaveformViewer.ToSharedRef(), &SFixedSampledSequenceViewer::OnStyleUpdated);
+	WaveformEditorStyle->OnNewWaveformViewerStyle.AddSP(WaveformViewer.ToSharedRef(), &SFixedSampledSequenceViewer::OnStyleUpdated);
 	GridData->OnGridMetricsUpdated.AddSP(WaveformViewer.ToSharedRef(), &SFixedSampledSequenceViewer::UpdateGridMetrics);
 }
 
 void STransformedWaveformViewPanel::SetUpGridData()
 {
-	FFixedSampleSequenceRulerStyle* RulerStyle = &WaveformEditorStyle->GetRegisteredWidgetStyle<FFixedSampleSequenceRulerStyle>("WaveformEditorRuler.Style").Get();
-	check(RulerStyle)
+	const FFixedSampleSequenceRulerStyle& RulerStyle = WaveformEditorStyle->GetWidgetStyle<FFixedSampleSequenceRulerStyle>("WaveformEditorRuler.Style");
 	
-	GridData = MakeShared<FWaveformEditorGridData>(DataView.SampleData.Num() / DataView.NumDimensions, DataView.SampleRate, RulerStyle->DesiredWidth, &RulerStyle->TicksTextFont);
+	GridData = MakeShared<FWaveformEditorGridData>(DataView.SampleData.Num() / DataView.NumDimensions, DataView.SampleRate, RulerStyle.DesiredWidth, &RulerStyle.TicksTextFont);
 }
 
 void STransformedWaveformViewPanel::ReceiveSequenceView(const FFixedSampledSequenceView InView, const uint32 FirstSampleIndex)
@@ -278,19 +271,9 @@ void STransformedWaveformViewPanel::UpdatePlayheadPosition(const float PaintedWi
 	}
 }
 
-void STransformedWaveformViewPanel::UpdateBackground(const FNotifyingAudioWidgetStyle& UpdatedStyle)
+void STransformedWaveformViewPanel::UpdateBackground(const FSampledSequenceViewerStyle UpdatedStyle)
 {
-	FSampledSequenceViewerStyle* ViewerStyle = &WaveformEditorStyle->GetRegisteredWidgetStyle<FSampledSequenceViewerStyle>("WaveformViewer.Style").Get();
-
-	check(ViewerStyle);
-
-	if (&UpdatedStyle != ViewerStyle)
-	{
-		return;
-	}
-
-	BackgroundBorder->SetBorderImage(&ViewerStyle->BackgroundBrush);
-	BackgroundBorder->SetBorderBackgroundColor(ViewerStyle->SequenceBackgroundColor);
+	BackgroundBorder->SetBorderBackgroundColor(UpdatedStyle.SequenceBackgroundColor);
 }
 
 void STransformedWaveformViewPanel::OnWaveEditorWidgetSettingsUpdated(const FName& PropertyName, const UWaveformEditorWidgetsSettings* Settings)
@@ -328,14 +311,13 @@ const UWaveformEditorWidgetsSettings* STransformedWaveformViewPanel::GetWaveform
 
 void STransformedWaveformViewPanel::SetUpBackground()
 {
-	FSampledSequenceViewerStyle* ViewerStyle = &WaveformEditorStyle->GetRegisteredWidgetStyle<FSampledSequenceViewerStyle>("WaveformViewer.Style").Get();
-	check(ViewerStyle);
+	WaveformEditorStyle->OnNewWaveformViewerStyle.AddSP(this, &STransformedWaveformViewPanel::UpdateBackground);
 
-	ViewerStyle->OnStyleUpdated.AddSP(this, &STransformedWaveformViewPanel::UpdateBackground);
+	const FSampledSequenceViewerStyle& ViewerStyle = WaveformEditorStyle->GetWidgetStyle<FSampledSequenceViewerStyle>("WaveformViewer.Style");
 
 	BackgroundBorder = SNew(SBorder)
-		.BorderImage(&ViewerStyle->BackgroundBrush)
-		.BorderBackgroundColor(ViewerStyle->SequenceBackgroundColor);
+		.BorderImage(&ViewerStyle.BackgroundBrush)
+		.BorderBackgroundColor(ViewerStyle.SequenceBackgroundColor);
 }
 
 void STransformedWaveformViewPanel::SetUpInputOverrides(const FArguments& InArgs)
@@ -368,7 +350,7 @@ void STransformedWaveformViewPanel::SetUpValueGridOverlay()
 	const uint32 NumValueGridDivision = Settings ? Settings->MaxLoudnessGridDivisions : 3;
 
 
-	FSampledSequenceValueGridOverlayStyle* ValueGridStyle = &WaveformEditorStyle->GetRegisteredWidgetStyle<FSampledSequenceValueGridOverlayStyle>("WaveformEditorValueGrid.Style").Get();
+	const FSampledSequenceValueGridOverlayStyle& ValueGridStyle = WaveformEditorStyle->GetWidgetStyle<FSampledSequenceValueGridOverlayStyle>("WaveformEditorValueGrid.Style");
 	
 	ValueGridOverlay = SNew(SSampledSequenceValueGridOverlay)
 		.DivideMode(SampledSequenceValueGridOverlay::EGridDivideMode::MidSplit)
@@ -377,9 +359,9 @@ void STransformedWaveformViewPanel::SetUpValueGridOverlay()
 		.MaxDivisionParameter(NumValueGridDivision)
 		.HideLabels(!bShowLabels)
 		.HideGrid(!bShowGrid)
-		.Style(ValueGridStyle);
+		.Style(&ValueGridStyle);
 
-	ValueGridStyle->OnStyleUpdated.AddSP(ValueGridOverlay.ToSharedRef(), &SSampledSequenceValueGridOverlay::OnStyleUpdated);
+	WaveformEditorStyle->OnNewValueGridOverlayStyle.AddSP(ValueGridOverlay.ToSharedRef(), &SSampledSequenceValueGridOverlay::OnStyleUpdated);
 }
 
 #undef LOCTEXT_NAMESPACE
diff --git a/Engine/Plugins/Editor/WaveformEditor/Source/WaveformEditorWidgets/Private/WaveformEditorStyle.cpp b/Engine/Plugins/Editor/WaveformEditor/Source/WaveformEditorWidgets/Private/WaveformEditorStyle.cpp
--- a/Engine/Plugins/Editor/WaveformEditor/Source/WaveformEditorWidgets/Private/WaveformEditorStyle.cpp
+++ b/Engine/Plugins/Editor/WaveformEditor/Source/WaveformEditorWidgets/Private/WaveformEditorStyle.cpp
@@ -2,7 +2,6 @@
 
 #include "WaveformEditorStyle.h"
 
-#include "AudioWidgetsSlateTypes.h"
 #include "Styling/SlateStyleRegistry.h"
 #include "WaveformEditorWidgetsSettings.h"
 
@@ -15,6 +14,11 @@ static FLazyName WaveformViewerStyleName("WaveformViewer.Style");
 FName FWaveformEditorStyle::StyleName("WaveformEditorStyle");
 TUniquePtr<FWaveformEditorStyle> FWaveformEditorStyle::StyleInstance = nullptr;
 
+FOnNewWaveformViewerStyle FWaveformEditorStyle::OnNewWaveformViewerStyle;
+FOnNewPlayheadOverlayStyle FWaveformEditorStyle::OnNewPlayheadOverlayStyle;
+FOnNewTimeRulerStyle FWaveformEditorStyle::OnNewTimeRulerStyle;
+FOnNewValueGridOverlayStyle FWaveformEditorStyle::OnNewValueGridOverlayStyle;
+
 FWaveformEditorStyle::FWaveformEditorStyle()
 	: FSlateStyleSet(StyleName)
 {
@@ -42,43 +46,22 @@ void FWaveformEditorStyle::Init()
 	
 	check(Settings);
 	Settings->OnSettingChanged().AddStatic(&FWaveformEditorStyle::OnWidgetSettingsUpdated);
-
 	StyleInstance->SetParentStyleName("CoreStyle");
 
 	//Waveform Viewer style
-	FSampledSequenceViewerStyle WaveViewerStyle = FSampledSequenceViewerStyle()
-		.SetSequenceColor(Settings->WaveformColor)
-		.SetBackgroundColor(Settings->WaveformBackgroundColor)
-		.SetSequenceLineThickness(Settings->WaveformLineThickness)
-		.SetSampleMarkersSize(Settings->SampleMarkersSize)
-		.SetMajorGridLineColor(Settings->MajorGridColor)
-		.SetMinorGridLineColor(Settings->MinorGridColor)
-		.SetZeroCrossingLineColor(Settings->LoudnessGridColor)
-		.SetZeroCrossingLineThickness(Settings->ZeroCrossingLineThickness);
-
+	FSampledSequenceViewerStyle WaveViewerStyle = CreateWaveformViewerStyleFromSettings(*Settings);
 	StyleInstance->Set(WaveformViewerStyleName, WaveViewerStyle);
 
 	//Playhead Overlay style
-	FPlayheadOverlayStyle PlayheadOverlayStyle = FPlayheadOverlayStyle().SetPlayheadColor(Settings->PlayheadColor);
+	FPlayheadOverlayStyle PlayheadOverlayStyle = CreatePlayheadOverlayStyleFromSettings(*Settings);
 	StyleInstance->Set(PlayheadOverlayStyleName, PlayheadOverlayStyle);
 
 	//Time Ruler style 
-	FFixedSampleSequenceRulerStyle TimeRulerStyle = FFixedSampleSequenceRulerStyle()
-		.SetHandleColor(Settings->PlayheadColor)
-		.SetTicksColor(Settings->RulerTicksColor)
-		.SetTicksTextColor(Settings->RulerTextColor)
-		.SetHandleColor(Settings->PlayheadColor)
-		.SetFontSize(Settings->RulerFontSize);
-
+	FFixedSampleSequenceRulerStyle TimeRulerStyle = CreateTimeRulerStyleFromSettings(*Settings);
 	StyleInstance->Set(WaveformEditorRulerStyleName, TimeRulerStyle);
 
 	//ValueGridOverlayStyle
-	FSampledSequenceValueGridOverlayStyle ValueGridOverlayStyle = FSampledSequenceValueGridOverlayStyle()
-		.SetGridColor(Settings->LoudnessGridColor)
-		.SetGridThickness(Settings->LoudnessGridThickness)
-		.SetLabelTextColor(Settings->LoudnessGridTextColor)
-		.SetLabelTextFontSize(Settings->LoudnessGridTextSize);
-
+	FSampledSequenceValueGridOverlayStyle ValueGridOverlayStyle = CreateValueGridOverlayStyleFromSettings(*Settings);
 	StyleInstance->Set(ValueGridOverlayStyleName, ValueGridOverlayStyle);
 
 	FSlateStyleRegistry::RegisterSlateStyle(StyleInstance->Get());
@@ -94,77 +77,127 @@ const UWaveformEditorWidgetsSettings* FWaveformEditorStyle::GetWidgetsSettings()
 
 void FWaveformEditorStyle::OnWidgetSettingsUpdated(const FName& PropertyName, const UWaveformEditorWidgetsSettings* Settings)
 {
-	TSharedRef<FSampledSequenceViewerStyle> WaveformViewerStyle = GetRegisteredWidgetStyle<FSampledSequenceViewerStyle>(WaveformViewerStyleName);
-	TSharedRef<FPlayheadOverlayStyle> PlayheadOverlayStyle = GetRegisteredWidgetStyle<FPlayheadOverlayStyle>(PlayheadOverlayStyleName);
-	TSharedRef<FFixedSampleSequenceRulerStyle> WaveformEditorTimeRulerStyle = GetRegisteredWidgetStyle<FFixedSampleSequenceRulerStyle>(WaveformEditorRulerStyleName);
-	TSharedRef<FSampledSequenceValueGridOverlayStyle> ValueGridOverlayStyle = GetRegisteredWidgetStyle<FSampledSequenceValueGridOverlayStyle>(ValueGridOverlayStyleName);
-
+	check(Settings)
 
 	if (PropertyName == GET_MEMBER_NAME_CHECKED(UWaveformEditorWidgetsSettings, WaveformColor))
 	{
-		WaveformViewerStyle->SetSequenceColor(Settings->WaveformColor);
+		CreateWaveformViewerStyleFromSettings(*Settings);
 	}
 	else if (PropertyName == GET_MEMBER_NAME_CHECKED(UWaveformEditorWidgetsSettings, WaveformBackgroundColor))
 	{
-		WaveformViewerStyle->SetBackgroundColor(Settings->WaveformBackgroundColor);
+		CreateWaveformViewerStyleFromSettings(*Settings);
 	}
 	else if (PropertyName == GET_MEMBER_NAME_CHECKED(UWaveformEditorWidgetsSettings, WaveformLineThickness))
 	{
-		WaveformViewerStyle->SetSequenceLineThickness(Settings->WaveformLineThickness);
+		CreateWaveformViewerStyleFromSettings(*Settings);
 	}
 	else if (PropertyName == GET_MEMBER_NAME_CHECKED(UWaveformEditorWidgetsSettings, SampleMarkersSize))
 	{
-		WaveformViewerStyle->SetSampleMarkersSize(Settings->SampleMarkersSize);
+		CreateWaveformViewerStyleFromSettings(*Settings);
 	}
 	else if (PropertyName == GET_MEMBER_NAME_CHECKED(UWaveformEditorWidgetsSettings, MajorGridColor))
 	{
-		WaveformViewerStyle->SetMajorGridLineColor(Settings->MajorGridColor);
+		CreateWaveformViewerStyleFromSettings(*Settings);
 	}
 	else if (PropertyName == GET_MEMBER_NAME_CHECKED(UWaveformEditorWidgetsSettings, MinorGridColor))
 	{
-		WaveformViewerStyle->SetMinorGridLineColor(Settings->MinorGridColor);
+		CreateWaveformViewerStyleFromSettings(*Settings);
 	}
 	else if (PropertyName == GET_MEMBER_NAME_CHECKED(UWaveformEditorWidgetsSettings, ZeroCrossingLineThickness))
 	{
-		WaveformViewerStyle->SetZeroCrossingLineThickness(Settings->ZeroCrossingLineThickness);
+		CreateWaveformViewerStyleFromSettings(*Settings);
 	}
 	else if (PropertyName == GET_MEMBER_NAME_CHECKED(UWaveformEditorWidgetsSettings, PlayheadColor))
 	{
-		PlayheadOverlayStyle->SetPlayheadColor(Settings->PlayheadColor);
-		WaveformEditorTimeRulerStyle->SetHandleColor(Settings->PlayheadColor);
+		CreatePlayheadOverlayStyleFromSettings(*Settings);
+		CreateTimeRulerStyleFromSettings(*Settings);
 	}
 	else if (PropertyName == GET_MEMBER_NAME_CHECKED(UWaveformEditorWidgetsSettings, RulerBackgroundColor))
 	{
-		WaveformEditorTimeRulerStyle->SetBackgroundColor(Settings->RulerBackgroundColor);
+		CreateTimeRulerStyleFromSettings(*Settings);
 	}
 	else if (PropertyName == GET_MEMBER_NAME_CHECKED(UWaveformEditorWidgetsSettings, RulerTicksColor))
 	{
-		WaveformEditorTimeRulerStyle->SetTicksColor(Settings->RulerTicksColor);
+		CreateTimeRulerStyleFromSettings(*Settings);
 	}
 	else if (PropertyName == GET_MEMBER_NAME_CHECKED(UWaveformEditorWidgetsSettings, RulerTextColor))
 	{
-		WaveformEditorTimeRulerStyle->SetTicksTextColor(Settings->RulerTextColor);
+		CreateTimeRulerStyleFromSettings(*Settings);
 	}
 	else if (PropertyName == GET_MEMBER_NAME_CHECKED(UWaveformEditorWidgetsSettings, RulerFontSize))
 	{
-		WaveformEditorTimeRulerStyle->SetFontSize(Settings->RulerFontSize);
+		CreateTimeRulerStyleFromSettings(*Settings);
 	}
 	else if (PropertyName == GET_MEMBER_NAME_CHECKED(UWaveformEditorWidgetsSettings, LoudnessGridColor))
 	{
-		ValueGridOverlayStyle->SetGridColor(Settings->LoudnessGridColor);
-		WaveformViewerStyle->SetZeroCrossingLineColor(Settings->LoudnessGridColor);
+		CreateValueGridOverlayStyleFromSettings(*Settings);
+		CreateWaveformViewerStyleFromSettings(*Settings);
 	}
 	else if (PropertyName == GET_MEMBER_NAME_CHECKED(UWaveformEditorWidgetsSettings, LoudnessGridThickness))
 	{
-		ValueGridOverlayStyle->SetGridThickness(Settings->LoudnessGridThickness);
+		CreateValueGridOverlayStyleFromSettings(*Settings);
 	}
 	else if (PropertyName == GET_MEMBER_NAME_CHECKED(UWaveformEditorWidgetsSettings, LoudnessGridTextColor))
 	{
-		ValueGridOverlayStyle->SetLabelTextColor(Settings->LoudnessGridTextColor);
+		CreateValueGridOverlayStyleFromSettings(*Settings);
 	}
 	else if (PropertyName == GET_MEMBER_NAME_CHECKED(UWaveformEditorWidgetsSettings, LoudnessGridTextSize))
 	{
-		ValueGridOverlayStyle->SetLabelTextFontSize(Settings->LoudnessGridTextSize);
+		CreateValueGridOverlayStyleFromSettings(*Settings);
 	}
 
 }
+
+FSampledSequenceViewerStyle FWaveformEditorStyle::CreateWaveformViewerStyleFromSettings(const UWaveformEditorWidgetsSettings& InSettings)
+{
+	FSampledSequenceViewerStyle WaveViewerStyle = FSampledSequenceViewerStyle()
+		.SetSequenceColor(InSettings.WaveformColor)
+		.SetBackgroundColor(InSettings.WaveformBackgroundColor)
+		.SetSequenceLineThickness(InSettings.WaveformLineThickness)
+		.SetSampleMarkersSize(InSettings.SampleMarkersSize)
+		.SetMajorGridLineColor(InSettings.MajorGridColor)
+		.SetMinorGridLineColor(InSettings.MinorGridColor)
+		.SetZeroCrossingLineColor(InSettings.LoudnessGridColor)
+		.SetZeroCrossingLineThickness(InSettings.ZeroCrossingLineThickness);
+
+	OnNewWaveformViewerStyle.Broadcast(WaveViewerStyle);
+
+	return WaveViewerStyle;
+}
+
+FPlayheadOverlayStyle FWaveformEditorStyle::CreatePlayheadOverlayStyleFromSettings(const UWaveformEditorWidgetsSettings& InSettings)
+{
+	FPlayheadOverlayStyle PlayheadOverlayStyle = FPlayheadOverlayStyle().SetPlayheadColor(InSettings.PlayheadColor);
+
+	OnNewPlayheadOverlayStyle.Broadcast(PlayheadOverlayStyle);
+
+	return PlayheadOverlayStyle;
+}
+
+FFixedSampleSequenceRulerStyle FWaveformEditorStyle::CreateTimeRulerStyleFromSettings(const UWaveformEditorWidgetsSettings& InSettings)
+{
+	FFixedSampleSequenceRulerStyle TimeRulerStyle = FFixedSampleSequenceRulerStyle()
+		.SetHandleColor(InSettings.PlayheadColor)
+		.SetTicksColor(InSettings.RulerTicksColor)
+		.SetTicksTextColor(InSettings.RulerTextColor)
+		.SetHandleColor(InSettings.PlayheadColor)
+		.SetFontSize(InSettings.RulerFontSize)
+		.SetBackgroundColor(InSettings.RulerBackgroundColor);
+
+	OnNewTimeRulerStyle.Broadcast(TimeRulerStyle);
+
+	return TimeRulerStyle;
+}
+
+FSampledSequenceValueGridOverlayStyle FWaveformEditorStyle::CreateValueGridOverlayStyleFromSettings(const UWaveformEditorWidgetsSettings& InSettings)
+{
+	FSampledSequenceValueGridOverlayStyle ValueGridOverlayStyle = FSampledSequenceValueGridOverlayStyle()
+		.SetGridColor(InSettings.LoudnessGridColor)
+		.SetGridThickness(InSettings.LoudnessGridThickness)
+		.SetLabelTextColor(InSettings.LoudnessGridTextColor)
+		.SetLabelTextFontSize(InSettings.LoudnessGridTextSize);
+
+	OnNewValueGridOverlayStyle.Broadcast(ValueGridOverlayStyle);
+
+	return ValueGridOverlayStyle;
+}
diff --git a/Engine/Plugins/Editor/WaveformEditor/Source/WaveformEditorWidgets/Public/STransformedWaveformViewPanel.h b/Engine/Plugins/Editor/WaveformEditor/Source/WaveformEditorWidgets/Public/STransformedWaveformViewPanel.h
--- a/Engine/Plugins/Editor/WaveformEditor/Source/WaveformEditorWidgets/Public/STransformedWaveformViewPanel.h
+++ b/Engine/Plugins/Editor/WaveformEditor/Source/WaveformEditorWidgets/Public/STransformedWaveformViewPanel.h
@@ -2,6 +2,7 @@
 
 #pragma once
 
+#include "AudioWidgetsSlateTypes.h"
 #include "HAL/Platform.h"
 #include "Widgets/SCompoundWidget.h"
 #include "IFixedSampledSequenceViewReceiver.h"
@@ -72,7 +73,7 @@ private:
 
 	void UpdateDisplayUnit(const ESampledSequenceDisplayUnit InDisplayUnit);
 	void UpdatePlayheadPosition(const float PaintedWidth);
-	void UpdateBackground(const FNotifyingAudioWidgetStyle& UpdatedStyle);
+	void UpdateBackground(const FSampledSequenceViewerStyle UpdatedStyle);
 
 	void OnWaveEditorWidgetSettingsUpdated(const FName& PropertyName, const UWaveformEditorWidgetsSettings* Settings);
 
diff --git a/Engine/Plugins/Editor/WaveformEditor/Source/WaveformEditorWidgets/Public/WaveformEditorStyle.h b/Engine/Plugins/Editor/WaveformEditor/Source/WaveformEditorWidgets/Public/WaveformEditorStyle.h
--- a/Engine/Plugins/Editor/WaveformEditor/Source/WaveformEditorWidgets/Public/WaveformEditorStyle.h
+++ b/Engine/Plugins/Editor/WaveformEditor/Source/WaveformEditorWidgets/Public/WaveformEditorStyle.h
@@ -2,10 +2,17 @@
 
 #pragma once
 
+#include "AudioWidgetsSlateTypes.h"
+#include "Delegates/Delegate.h"
 #include "Styling/SlateStyle.h"
 
 class UWaveformEditorWidgetsSettings;
 
+DECLARE_MULTICAST_DELEGATE_OneParam(FOnNewPlayheadOverlayStyle, FPlayheadOverlayStyle)
+DECLARE_MULTICAST_DELEGATE_OneParam(FOnNewTimeRulerStyle, FFixedSampleSequenceRulerStyle)
+DECLARE_MULTICAST_DELEGATE_OneParam(FOnNewValueGridOverlayStyle, FSampledSequenceValueGridOverlayStyle)
+DECLARE_MULTICAST_DELEGATE_OneParam(FOnNewWaveformViewerStyle, FSampledSequenceViewerStyle)
+
 class FWaveformEditorStyle
 	: public FSlateStyleSet
 
@@ -20,19 +27,19 @@ public:
 	static FWaveformEditorStyle& Get();
 	static void Init();
 
-	template< typename SlateWidgetStyle >
-	static TSharedRef<SlateWidgetStyle> GetRegisteredWidgetStyle(const FName& InStyleName)
-	{
-		TSharedRef<FSlateWidgetStyle>* WidgetStyleBase = StyleInstance->WidgetStyleValues.Find(InStyleName);
-		check(WidgetStyleBase);
-		TSharedRef<SlateWidgetStyle> WidgetStyle = StaticCastSharedRef<SlateWidgetStyle>(*WidgetStyleBase);
-		return WidgetStyle;
-	}
+	static FOnNewPlayheadOverlayStyle OnNewPlayheadOverlayStyle;
+	static FOnNewTimeRulerStyle OnNewTimeRulerStyle;
+	static FOnNewValueGridOverlayStyle OnNewValueGridOverlayStyle;
+	static FOnNewWaveformViewerStyle OnNewWaveformViewerStyle;
 
-	
 private:
 	static TUniquePtr<FWaveformEditorStyle> StyleInstance;
 
 	static const UWaveformEditorWidgetsSettings* GetWidgetsSettings();
 	static void  OnWidgetSettingsUpdated(const FName& PropertyName, const UWaveformEditorWidgetsSettings* Settings);
+
+	static FFixedSampleSequenceRulerStyle CreateTimeRulerStyleFromSettings(const UWaveformEditorWidgetsSettings& InSettings);
+	static FPlayheadOverlayStyle CreatePlayheadOverlayStyleFromSettings(const UWaveformEditorWidgetsSettings& InSettings);
+	static FSampledSequenceValueGridOverlayStyle CreateValueGridOverlayStyleFromSettings(const UWaveformEditorWidgetsSettings& InSettings);
+	static FSampledSequenceViewerStyle CreateWaveformViewerStyleFromSettings(const UWaveformEditorWidgetsSettings& InSettings);
 };
\ No newline at end of file
diff --git a/Engine/Plugins/Runtime/AudioWidgets/Source/AudioWidgets/Private/SFixedSampledSequenceRuler.cpp b/Engine/Plugins/Runtime/AudioWidgets/Source/AudioWidgets/Private/SFixedSampledSequenceRuler.cpp
--- a/Engine/Plugins/Runtime/AudioWidgets/Source/AudioWidgets/Private/SFixedSampledSequenceRuler.cpp
+++ b/Engine/Plugins/Runtime/AudioWidgets/Source/AudioWidgets/Private/SFixedSampledSequenceRuler.cpp
@@ -19,20 +19,19 @@ void SFixedSampledSequenceRuler::Construct(const FArguments& InArgs, TSharedRef<
 		bDisplayPlayhead = InArgs._DisplayPlayhead.Get();
 	}
 
-	Style = InArgs._Style;
-	check(Style);
+	check(InArgs._Style);
 
-	HandleColor = Style->HandleColor;
-	HandleWidth = Style->HandleWidth;
-	HandleBrush = Style->HandleBrush;
-	BackgroundColor = Style->BackgroundColor;
-	BackgroundBrush = Style->BackgroundBrush;
-	TicksColor = Style->TicksColor;
-	TicksTextColor = Style->TicksTextColor;
-	TicksTextFont = Style->TicksTextFont;
-	DesiredWidth = Style->DesiredWidth;
-	DesiredHeight = Style->DesiredHeight;
-	TicksTextOffset = Style->TicksTextOffset;
+	HandleColor = InArgs._Style->HandleColor;
+	HandleWidth = InArgs._Style->HandleWidth;
+	HandleBrush = InArgs._Style->HandleBrush;
+	BackgroundColor = InArgs._Style->BackgroundColor;
+	BackgroundBrush = InArgs._Style->BackgroundBrush;
+	TicksColor = InArgs._Style->TicksColor;
+	TicksTextColor = InArgs._Style->TicksTextColor;
+	TicksTextFont = InArgs._Style->TicksTextFont;
+	DesiredWidth = InArgs._Style->DesiredWidth;
+	DesiredHeight = InArgs._Style->DesiredHeight;
+	TicksTextOffset = InArgs._Style->TicksTextOffset;
 }
 
 FVector2D SFixedSampledSequenceRuler::ComputeDesiredSize(float) const
@@ -221,25 +220,19 @@ void SFixedSampledSequenceRuler::SetPlayheadPosition(const float InNewPosition)
 	PlayheadPosition = InNewPosition;
 }
 
-void SFixedSampledSequenceRuler::OnStyleUpdated(const FNotifyingAudioWidgetStyle& UpdatedStyle)
+void SFixedSampledSequenceRuler::OnStyleUpdated(const FFixedSampleSequenceRulerStyle UpdatedStyle)
 {
-	check(Style);
-	if (&UpdatedStyle != Style)
-	{
-		return;
-	}
-
-	HandleColor = Style->HandleColor;
-	HandleWidth = Style->HandleWidth;
-	HandleBrush = Style->HandleBrush;
-	BackgroundColor = Style->BackgroundColor;
-	BackgroundBrush = Style->BackgroundBrush;
-	TicksColor = Style->TicksColor;
-	TicksTextColor = Style->TicksTextColor;
-	TicksTextFont = Style->TicksTextFont;
-	DesiredWidth = Style->DesiredWidth;
-	DesiredHeight = Style->DesiredHeight;
-	TicksTextOffset = Style->TicksTextOffset;
+	HandleColor = UpdatedStyle.HandleColor;
+	HandleWidth = UpdatedStyle.HandleWidth;
+	HandleBrush = UpdatedStyle.HandleBrush;
+	BackgroundColor = UpdatedStyle.BackgroundColor;
+	BackgroundBrush = UpdatedStyle.BackgroundBrush;
+	TicksColor = UpdatedStyle.TicksColor;
+	TicksTextColor = UpdatedStyle.TicksTextColor;
+	TicksTextFont = UpdatedStyle.TicksTextFont;
+	DesiredWidth = UpdatedStyle.DesiredWidth;
+	DesiredHeight = UpdatedStyle.DesiredHeight;
+	TicksTextOffset = UpdatedStyle.TicksTextOffset;
 }
 
 FReply SFixedSampledSequenceRuler::LaunchContextMenu()
diff --git a/Engine/Plugins/Runtime/AudioWidgets/Source/AudioWidgets/Private/SFixedSampledSequenceViewer.cpp b/Engine/Plugins/Runtime/AudioWidgets/Source/AudioWidgets/Private/SFixedSampledSequenceViewer.cpp
--- a/Engine/Plugins/Runtime/AudioWidgets/Source/AudioWidgets/Private/SFixedSampledSequenceViewer.cpp
+++ b/Engine/Plugins/Runtime/AudioWidgets/Source/AudioWidgets/Private/SFixedSampledSequenceViewer.cpp
@@ -7,23 +7,22 @@ void SFixedSampledSequenceViewer::Construct(const FArguments& InArgs, TArrayView
 	GridService = InGridService;
 	UpdateView(InSampleData, InNumChannels);
 
-	Style = InArgs._Style;
 	DrawingParams = InArgs._SequenceDrawingParams;
 	bHideBackground = InArgs._HideBackground;
 	bHideGrid = InArgs._HideGrid;
 
-	check(Style);
-	SequenceColor = Style->SequenceColor;
-	MajorGridLineColor = Style->MajorGridLineColor;
-	MinorGridLineColor = Style->MinorGridLineColor;
-	BackgroundColor = Style->SequenceBackgroundColor;
-	BackgroundBrush = Style->BackgroundBrush;
-	DesiredWidth = Style->DesiredWidth;
-	DesiredHeight = Style->DesiredHeight;
-	SampleMarkersSize = Style->SampleMarkersSize;
-	SequenceLineThickness = Style->SequenceLineThickness;
-	ZeroCrossingLineColor = Style->ZeroCrossingLineColor;
-	ZeroCrossingLineThickness = Style->ZeroCrossingLineThickness;
+	check(InArgs._Style);
+	SequenceColor = InArgs._Style->SequenceColor;
+	MajorGridLineColor = InArgs._Style->MajorGridLineColor;
+	MinorGridLineColor = InArgs._Style->MinorGridLineColor;
+	BackgroundColor = InArgs._Style->SequenceBackgroundColor;
+	BackgroundBrush = InArgs._Style->BackgroundBrush;
+	DesiredWidth = InArgs._Style->DesiredWidth;
+	DesiredHeight = InArgs._Style->DesiredHeight;
+	SampleMarkersSize = InArgs._Style->SampleMarkersSize;
+	SequenceLineThickness = InArgs._Style->SequenceLineThickness;
+	ZeroCrossingLineColor = InArgs._Style->ZeroCrossingLineColor;
+	ZeroCrossingLineThickness = InArgs._Style->ZeroCrossingLineThickness;
 }
 
 void SFixedSampledSequenceViewer::UpdateView(TArrayView<const float> InSampleData, const uint8 InNumChannels)
@@ -76,7 +75,9 @@ int32 SFixedSampledSequenceViewer::OnPaint(const FPaintArgs& Args, const FGeomet
 					AllottedGeometry.ToPaintGeometry(),
 					BinDrawPoints,
 					ESlateDrawEffect::None,
-					SequenceColor.GetSpecifiedColor()
+					SequenceColor.GetSpecifiedColor(),
+					true,
+					SequenceLineThickness
 				);
 
 				const int32 NextBinCoordinates = i + NumChannels;
@@ -92,7 +93,9 @@ int32 SFixedSampledSequenceViewer::OnPaint(const FPaintArgs& Args, const FGeomet
 						AllottedGeometry.ToPaintGeometry(),
 						BinDrawPoints,
 						ESlateDrawEffect::None,
-						SequenceColor.GetSpecifiedColor()
+						SequenceColor.GetSpecifiedColor(),
+						true,
+						SequenceLineThickness
 					);
 
 				}
@@ -268,24 +271,17 @@ void SFixedSampledSequenceViewer::UpdateGridMetrics()
 	GridMetrics = GridService->GetGridMetrics();
 }
 
-void SFixedSampledSequenceViewer::OnStyleUpdated(const FNotifyingAudioWidgetStyle& UpdatedStyle)
+void SFixedSampledSequenceViewer::OnStyleUpdated(const FSampledSequenceViewerStyle UpdatedStyle)
 {
-	check(Style);
-
-	if (&UpdatedStyle != Style)
-	{
-		return;
-	}
-
-	SequenceColor = Style->SequenceColor;
-	MajorGridLineColor = Style->MajorGridLineColor;
-	MinorGridLineColor = Style->MinorGridLineColor;
-	ZeroCrossingLineColor = Style->ZeroCrossingLineColor;
-	ZeroCrossingLineThickness = Style->ZeroCrossingLineThickness;
-	BackgroundColor = Style->SequenceBackgroundColor;
-	BackgroundBrush = Style->BackgroundBrush;
-	SampleMarkersSize = Style->SampleMarkersSize;
-	SequenceLineThickness = Style->SequenceLineThickness;
-	DesiredWidth = Style->DesiredWidth;
-	DesiredHeight = Style->DesiredHeight;
+	SequenceColor = UpdatedStyle.SequenceColor;
+	MajorGridLineColor = UpdatedStyle.MajorGridLineColor;
+	MinorGridLineColor = UpdatedStyle.MinorGridLineColor;
+	ZeroCrossingLineColor = UpdatedStyle.ZeroCrossingLineColor;
+	ZeroCrossingLineThickness = UpdatedStyle.ZeroCrossingLineThickness;
+	BackgroundColor = UpdatedStyle.SequenceBackgroundColor;
+	BackgroundBrush = UpdatedStyle.BackgroundBrush;
+	SampleMarkersSize = UpdatedStyle.SampleMarkersSize;
+	SequenceLineThickness = UpdatedStyle.SequenceLineThickness;
+	DesiredWidth = UpdatedStyle.DesiredWidth;
+	DesiredHeight = UpdatedStyle.DesiredHeight;
 }
\ No newline at end of file
diff --git a/Engine/Plugins/Runtime/AudioWidgets/Source/AudioWidgets/Private/SFixedSampledSequenceViewerPanel.cpp b/Engine/Plugins/Runtime/AudioWidgets/Source/AudioWidgets/Private/SFixedSampledSequenceViewerPanel.cpp
--- a/Engine/Plugins/Runtime/AudioWidgets/Source/AudioWidgets/Private/SFixedSampledSequenceViewerPanel.cpp
+++ b/Engine/Plugins/Runtime/AudioWidgets/Source/AudioWidgets/Private/SFixedSampledSequenceViewerPanel.cpp
@@ -170,6 +170,30 @@ void SFixedSampledSequenceViewerPanel::SetPlayheadRatio(const float InRatio)
 	CachedPlayheadRatio = InRatio;
 }
 
+void SFixedSampledSequenceViewerPanel::UpdateSequenceViewerStyle(const FSampledSequenceViewerStyle UpdatedStyle)
+{
+	check(SequenceViewer);
+	SequenceViewer->OnStyleUpdated(UpdatedStyle);
+}
+
+void SFixedSampledSequenceViewerPanel::UpdateRulerStyle(const FFixedSampleSequenceRulerStyle UpdatedStyle)
+{
+	check(SamplesRuler);
+	SamplesRuler->OnStyleUpdated(UpdatedStyle);
+}
+
+void SFixedSampledSequenceViewerPanel::UpdatePlayheadOverlayStyle(const FPlayheadOverlayStyle UpdatedStyle)
+{
+	check(PlayheadOverlay);
+	PlayheadOverlay->OnStyleUpdated(UpdatedStyle);
+}
+
+void SFixedSampledSequenceViewerPanel::UpdateValueGridOverlayStyle(const FSampledSequenceValueGridOverlayStyle UpdatedStyle)
+{
+	check(ValueGridOverlay);
+	ValueGridOverlay->OnStyleUpdated(UpdatedStyle);
+}
+
 void SFixedSampledSequenceViewerPanel::UpdateSamplesRulerDisplayUnit(const ESampledSequenceDisplayUnit InDisplayUnit)
 {
 	SamplesRulerDisplayUnit = InDisplayUnit;
diff --git a/Engine/Plugins/Runtime/AudioWidgets/Source/AudioWidgets/Private/SPlayheadOverlay.cpp b/Engine/Plugins/Runtime/AudioWidgets/Source/AudioWidgets/Private/SPlayheadOverlay.cpp
--- a/Engine/Plugins/Runtime/AudioWidgets/Source/AudioWidgets/Private/SPlayheadOverlay.cpp
+++ b/Engine/Plugins/Runtime/AudioWidgets/Source/AudioWidgets/Private/SPlayheadOverlay.cpp
@@ -4,13 +4,11 @@
 
 void SPlayheadOverlay::Construct(const FArguments& InArgs)
 {
-	Style = InArgs._Style;
-
-	check(Style);
-	PlayheadColor = Style->PlayheadColor;
-	PlayheadWidth = Style->PlayheadWidth;
-	DesiredWidth = Style->DesiredWidth;
-	DesiredHeight = Style->DesiredHeight;
+	check(InArgs._Style);
+	PlayheadColor = InArgs._Style->PlayheadColor;
+	PlayheadWidth = InArgs._Style->PlayheadWidth;
+	DesiredWidth = InArgs._Style->DesiredWidth;
+	DesiredHeight = InArgs._Style->DesiredHeight;
 }
 
 void SPlayheadOverlay::SetPlayheadPosition(const float InNewPosition)
@@ -18,19 +16,12 @@ void SPlayheadOverlay::SetPlayheadPosition(const float InNewPosition)
 	PlayheadPosition = InNewPosition;
 }
 
-void SPlayheadOverlay::OnStyleUpdated(const FNotifyingAudioWidgetStyle& UpdatedStyle)
+void SPlayheadOverlay::OnStyleUpdated(const FPlayheadOverlayStyle UpdatedStyle)
 {
-	check(Style);
-
-	if (&UpdatedStyle != Style)
-	{
-		return;
-	}
-
-	PlayheadColor = Style->PlayheadColor;
-	PlayheadWidth = Style->PlayheadWidth;
-	DesiredWidth = Style->DesiredWidth;
-	DesiredHeight = Style->DesiredHeight;
+	PlayheadColor = UpdatedStyle.PlayheadColor;
+	PlayheadWidth = UpdatedStyle.PlayheadWidth;
+	DesiredWidth = UpdatedStyle.DesiredWidth;
+	DesiredHeight = UpdatedStyle.DesiredHeight;
 }
 
 int32 SPlayheadOverlay::OnPaint(const FPaintArgs& Args, const FGeometry& AllottedGeometry, const FSlateRect& MyCullingRect, FSlateWindowElementList& OutDrawElements, int32 LayerId, const FWidgetStyle& InWidgetStyle, bool bParentEnabled) const
diff --git a/Engine/Plugins/Runtime/AudioWidgets/Source/AudioWidgets/Private/SSampledSequenceValueGridOverlay.cpp b/Engine/Plugins/Runtime/AudioWidgets/Source/AudioWidgets/Private/SSampledSequenceValueGridOverlay.cpp
--- a/Engine/Plugins/Runtime/AudioWidgets/Source/AudioWidgets/Private/SSampledSequenceValueGridOverlay.cpp
+++ b/Engine/Plugins/Runtime/AudioWidgets/Source/AudioWidgets/Private/SSampledSequenceValueGridOverlay.cpp
@@ -25,14 +25,12 @@ void SSampledSequenceValueGridOverlay::Construct(const FArguments& InArgs)
 	DrawingParams = InArgs._SequenceDrawingParams;
 	bHideLabels = InArgs._HideLabels;
 	bHideGrid = InArgs._HideGrid;
-
-	Style = InArgs._Style;
-	GridColor = Style->GridColor;
-	LabelTextFont = Style->LabelTextFont;
-	LabelTextColor = Style->LabelTextColor;
-	DesiredWidth = Style->DesiredWidth;
-	DesiredHeight = Style->DesiredHeight;
-	GridThickness = Style->GridThickness;
+	GridColor = InArgs._Style->GridColor;
+	LabelTextFont = InArgs._Style->LabelTextFont;
+	LabelTextColor = InArgs._Style->LabelTextColor;
+	DesiredWidth = InArgs._Style->DesiredWidth;
+	DesiredHeight = InArgs._Style->DesiredHeight;
+	GridThickness = InArgs._Style->GridThickness;
 }
 
 int32 SSampledSequenceValueGridOverlay::OnPaint(const FPaintArgs& Args, const FGeometry& AllottedGeometry, const FSlateRect& MyCullingRect, FSlateWindowElementList& OutDrawElements, int32 LayerId, const FWidgetStyle& InWidgetStyle, bool bParentEnabled) const
@@ -129,21 +127,14 @@ void SSampledSequenceValueGridOverlay::Tick(const FGeometry& AllottedGeometry, c
 	}
 }
 
-void SSampledSequenceValueGridOverlay::OnStyleUpdated(const FNotifyingAudioWidgetStyle& UpdatedStyle)
+void SSampledSequenceValueGridOverlay::OnStyleUpdated(const FSampledSequenceValueGridOverlayStyle UpdatedStyle)
 {
-	check(Style);
-
-	if (&UpdatedStyle != Style)
-	{
-		return;
-	}
-
-	GridColor = Style->GridColor;
-	LabelTextFont = Style->LabelTextFont;
-	LabelTextColor = Style->LabelTextColor;
-	DesiredWidth = Style->DesiredWidth;
-	DesiredHeight = Style->DesiredHeight;
-	GridThickness = Style->GridThickness;
+	GridColor = UpdatedStyle.GridColor;
+	LabelTextFont = UpdatedStyle.LabelTextFont;
+	LabelTextColor = UpdatedStyle.LabelTextColor;
+	DesiredWidth = UpdatedStyle.DesiredWidth;
+	DesiredHeight = UpdatedStyle.DesiredHeight;
+	GridThickness = UpdatedStyle.GridThickness;
 
 	ForceRedraw();
 }
diff --git a/Engine/Plugins/Runtime/AudioWidgets/Source/AudioWidgets/Public/AudioWidgetsSlateTypes.h b/Engine/Plugins/Runtime/AudioWidgets/Source/AudioWidgets/Public/AudioWidgetsSlateTypes.h
--- a/Engine/Plugins/Runtime/AudioWidgets/Source/AudioWidgets/Public/AudioWidgetsSlateTypes.h
+++ b/Engine/Plugins/Runtime/AudioWidgets/Source/AudioWidgets/Public/AudioWidgetsSlateTypes.h
@@ -7,21 +7,6 @@
 
 #include "AudioWidgetsSlateTypes.generated.h"
 
-DECLARE_MULTICAST_DELEGATE_OneParam(FOnAudioWidgetStyleUpdated, const FNotifyingAudioWidgetStyle& /*Updated Widget Style*/);
-
-USTRUCT()
-struct FNotifyingAudioWidgetStyle : public FSlateWidgetStyle
-{
-	GENERATED_USTRUCT_BODY()
-
-	FNotifyingAudioWidgetStyle() = default;
-	virtual ~FNotifyingAudioWidgetStyle() = default;
-
-	virtual void BroadcastStyleUpdate() const { OnStyleUpdated.Broadcast(*this); }
-
-	FOnAudioWidgetStyleUpdated OnStyleUpdated;
-};
-
 /**
  * Represents the appearance of an Audio Text Box 
  */
@@ -190,7 +175,7 @@ struct AUDIOWIDGETS_API FAudioRadialSliderStyle : public FSlateWidgetStyle
  * Represents the appearance of a Sampled Sequence Viewer
  */
 USTRUCT(BlueprintType)
-struct AUDIOWIDGETS_API FSampledSequenceViewerStyle : public FNotifyingAudioWidgetStyle
+struct AUDIOWIDGETS_API FSampledSequenceViewerStyle : public FSlateWidgetStyle
 {
 	GENERATED_USTRUCT_BODY()
 
@@ -204,54 +189,54 @@ struct AUDIOWIDGETS_API FSampledSequenceViewerStyle : public FNotifyingAudioWidg
 
 	UPROPERTY(EditAnywhere, BlueprintReadWrite, Category = Appearance)
 	FSlateColor SequenceColor;
-	FSampledSequenceViewerStyle& SetSequenceColor(const FSlateColor InSequenceColor) { SequenceColor = InSequenceColor; BroadcastStyleUpdate(); return *this; }
+	FSampledSequenceViewerStyle& SetSequenceColor(const FSlateColor InSequenceColor) { SequenceColor = InSequenceColor; return *this; }
 
 	UPROPERTY(EditAnywhere, BlueprintReadWrite, Category = Appearance)
 	float SequenceLineThickness;
-	FSampledSequenceViewerStyle& SetSequenceLineThickness(const float InSequenceLineThickness) { SequenceLineThickness = InSequenceLineThickness; BroadcastStyleUpdate(); return *this; }
+	FSampledSequenceViewerStyle& SetSequenceLineThickness(const float InSequenceLineThickness) { SequenceLineThickness = InSequenceLineThickness; return *this; }
 
 	UPROPERTY(EditAnywhere, BlueprintReadWrite, Category = Appearance)
 	FSlateColor MajorGridLineColor;
-	FSampledSequenceViewerStyle& SetMajorGridLineColor(const FSlateColor InMajorGridLineColor) { MajorGridLineColor = InMajorGridLineColor; BroadcastStyleUpdate(); return *this; }
+	FSampledSequenceViewerStyle& SetMajorGridLineColor(const FSlateColor InMajorGridLineColor) { MajorGridLineColor = InMajorGridLineColor; return *this; }
 
 	UPROPERTY(EditAnywhere, BlueprintReadWrite, Category = Appearance)
 	FSlateColor MinorGridLineColor;
-	FSampledSequenceViewerStyle& SetMinorGridLineColor(const FSlateColor InMinorGridLineColor) { MinorGridLineColor = InMinorGridLineColor; BroadcastStyleUpdate(); return *this; }
+	FSampledSequenceViewerStyle& SetMinorGridLineColor(const FSlateColor InMinorGridLineColor) { MinorGridLineColor = InMinorGridLineColor; return *this; }
 
 	UPROPERTY(EditAnywhere, BlueprintReadWrite, Category = Appearance)
 	FSlateColor ZeroCrossingLineColor;
-	FSampledSequenceViewerStyle& SetZeroCrossingLineColor(const FSlateColor InZeroCrossingLineColor) { ZeroCrossingLineColor = InZeroCrossingLineColor; BroadcastStyleUpdate(); return *this; }
+	FSampledSequenceViewerStyle& SetZeroCrossingLineColor(const FSlateColor InZeroCrossingLineColor) { ZeroCrossingLineColor = InZeroCrossingLineColor; return *this; }
 
 	UPROPERTY(EditAnywhere, BlueprintReadWrite, Category = Appearance)
 	float ZeroCrossingLineThickness;
-	FSampledSequenceViewerStyle& SetZeroCrossingLineThickness(const float InZeroCrossingLineThickness) { ZeroCrossingLineThickness = InZeroCrossingLineThickness; BroadcastStyleUpdate(); return *this; }
+	FSampledSequenceViewerStyle& SetZeroCrossingLineThickness(const float InZeroCrossingLineThickness) { ZeroCrossingLineThickness = InZeroCrossingLineThickness; return *this; }
 
 	UPROPERTY(EditAnywhere, BlueprintReadWrite, Category = Appearance)
 	float SampleMarkersSize;
-	FSampledSequenceViewerStyle& SetSampleMarkersSize(const float InSampleMarkersSize) { SampleMarkersSize = InSampleMarkersSize; BroadcastStyleUpdate(); return *this; }
+	FSampledSequenceViewerStyle& SetSampleMarkersSize(const float InSampleMarkersSize) { SampleMarkersSize = InSampleMarkersSize; return *this; }
 
 	UPROPERTY(EditAnywhere, BlueprintReadWrite, Category = Appearance)
 	FSlateColor SequenceBackgroundColor;
-	FSampledSequenceViewerStyle& SetBackgroundColor(const FSlateColor InBackgroundColor) { SequenceBackgroundColor = InBackgroundColor; BroadcastStyleUpdate(); return *this; }
+	FSampledSequenceViewerStyle& SetBackgroundColor(const FSlateColor InBackgroundColor) { SequenceBackgroundColor = InBackgroundColor; return *this; }
 
 	UPROPERTY(EditAnywhere, BlueprintReadWrite, Category = Appearance)
 	FSlateBrush BackgroundBrush;
-	FSampledSequenceViewerStyle& SetBackgroundBrush(const FSlateBrush InBackgroundBrush) { BackgroundBrush = InBackgroundBrush; BroadcastStyleUpdate(); return *this; }
+	FSampledSequenceViewerStyle& SetBackgroundBrush(const FSlateBrush InBackgroundBrush) { BackgroundBrush = InBackgroundBrush; return *this; }
 
 	UPROPERTY(EditAnywhere, BlueprintReadWrite, Category = Appearance)
 	float DesiredWidth;
-	FSampledSequenceViewerStyle& SetDesiredWidth(const float InDesiredWidth) { DesiredWidth = InDesiredWidth; BroadcastStyleUpdate(); return *this; }
+	FSampledSequenceViewerStyle& SetDesiredWidth(const float InDesiredWidth) { DesiredWidth = InDesiredWidth; return *this; }
 
 	UPROPERTY(EditAnywhere, BlueprintReadWrite, Category = Appearance)
 	float DesiredHeight;
-	FSampledSequenceViewerStyle& SetDesiredHeight(const float InDesiredHeight) { DesiredHeight = InDesiredHeight; BroadcastStyleUpdate(); return *this; }
+	FSampledSequenceViewerStyle& SetDesiredHeight(const float InDesiredHeight) { DesiredHeight = InDesiredHeight; return *this; }
 };
 
 /**
  * Represents the appearance of a Waveform Viewer Overlay style
  */
 USTRUCT(BlueprintType)
-struct AUDIOWIDGETS_API FPlayheadOverlayStyle : public FNotifyingAudioWidgetStyle
+struct AUDIOWIDGETS_API FPlayheadOverlayStyle : public FSlateWidgetStyle
 {
 	GENERATED_USTRUCT_BODY()
 
@@ -264,26 +249,26 @@ struct AUDIOWIDGETS_API FPlayheadOverlayStyle : public FNotifyingAudioWidgetStyl
 
 	UPROPERTY(EditAnywhere, BlueprintReadWrite, Category = Appearance)
 	FSlateColor PlayheadColor;
-	FPlayheadOverlayStyle& SetPlayheadColor(const FSlateColor InPlayheadColor) { PlayheadColor = InPlayheadColor; BroadcastStyleUpdate(); return *this; }
+	FPlayheadOverlayStyle& SetPlayheadColor(const FSlateColor InPlayheadColor) { PlayheadColor = InPlayheadColor; return *this; }
 
 	UPROPERTY(EditAnywhere, BlueprintReadWrite, Category = Appearance)
 	float PlayheadWidth;
-	FPlayheadOverlayStyle& SetPlayheadWidth(const float InPlayheadWidth) { PlayheadWidth = InPlayheadWidth; BroadcastStyleUpdate(); return *this; }
+	FPlayheadOverlayStyle& SetPlayheadWidth(const float InPlayheadWidth) { PlayheadWidth = InPlayheadWidth; return *this; }
 
 	UPROPERTY(EditAnywhere, BlueprintReadWrite, Category = Appearance)
 	float DesiredWidth;
-	FPlayheadOverlayStyle& SetDesiredWidth(const float InDesiredWidth) { DesiredWidth = InDesiredWidth; BroadcastStyleUpdate(); return *this; }
+	FPlayheadOverlayStyle& SetDesiredWidth(const float InDesiredWidth) { DesiredWidth = InDesiredWidth; return *this; }
 
 	UPROPERTY(EditAnywhere, BlueprintReadWrite, Category = Appearance)
 	float DesiredHeight;
-	FPlayheadOverlayStyle& SetDesiredHeight(const float InDesiredHeight) { DesiredHeight = InDesiredHeight; BroadcastStyleUpdate(); return *this; }
+	FPlayheadOverlayStyle& SetDesiredHeight(const float InDesiredHeight) { DesiredHeight = InDesiredHeight; return *this; }
 };
 
 /**
  * Represents the appearance of a Sampled Sequence Time Ruler
  */
 USTRUCT(BlueprintType)
-struct AUDIOWIDGETS_API FFixedSampleSequenceRulerStyle : public FNotifyingAudioWidgetStyle
+struct AUDIOWIDGETS_API FFixedSampleSequenceRulerStyle : public FSlateWidgetStyle
 {
 	GENERATED_USTRUCT_BODY()
 
@@ -297,55 +282,55 @@ struct AUDIOWIDGETS_API FFixedSampleSequenceRulerStyle : public FNotifyingAudioW
 
 	UPROPERTY(EditAnywhere, BlueprintReadWrite, Category = Appearance)
 	float HandleWidth;
-	FFixedSampleSequenceRulerStyle& SetHandleWidth(const float InHandleWidth) { HandleWidth = InHandleWidth; BroadcastStyleUpdate(); return *this; }
+	FFixedSampleSequenceRulerStyle& SetHandleWidth(const float InHandleWidth) { HandleWidth = InHandleWidth; return *this; }
 
 	UPROPERTY(EditAnywhere, BlueprintReadWrite, Category = Appearance)
 	FSlateColor HandleColor;
-	FFixedSampleSequenceRulerStyle& SetHandleColor(const FSlateColor& InHandleColor) { HandleColor = InHandleColor; BroadcastStyleUpdate(); return *this; }
+	FFixedSampleSequenceRulerStyle& SetHandleColor(const FSlateColor& InHandleColor) { HandleColor = InHandleColor; return *this; }
 
 	UPROPERTY(EditAnywhere, BlueprintReadWrite, Category = Appearance)
 	FSlateBrush HandleBrush;
-	FFixedSampleSequenceRulerStyle& SetHandleBrush(const FSlateBrush& InHandleBrush) { HandleBrush = InHandleBrush; BroadcastStyleUpdate(); return *this; }
+	FFixedSampleSequenceRulerStyle& SetHandleBrush(const FSlateBrush& InHandleBrush) { HandleBrush = InHandleBrush; return *this; }
 
 	UPROPERTY(EditAnywhere, BlueprintReadWrite, Category = Appearance)
 	FSlateColor TicksColor;
-	FFixedSampleSequenceRulerStyle& SetTicksColor(const FSlateColor& InTicksColor) { TicksColor = InTicksColor; BroadcastStyleUpdate(); return *this; }
+	FFixedSampleSequenceRulerStyle& SetTicksColor(const FSlateColor& InTicksColor) { TicksColor = InTicksColor; return *this; }
 
 	UPROPERTY(EditAnywhere, BlueprintReadWrite, Category = Appearance)
 	FSlateColor TicksTextColor;
-	FFixedSampleSequenceRulerStyle& SetTicksTextColor(const FSlateColor& InTicksTextColor) { TicksTextColor = InTicksTextColor; BroadcastStyleUpdate(); return *this; }
+	FFixedSampleSequenceRulerStyle& SetTicksTextColor(const FSlateColor& InTicksTextColor) { TicksTextColor = InTicksTextColor; return *this; }
 
 	UPROPERTY(EditAnywhere, BlueprintReadWrite, Category = Appearance)
 	FSlateFontInfo TicksTextFont;
-	FFixedSampleSequenceRulerStyle& SetTicksTextFont(const FSlateFontInfo& InTicksTextFont) { TicksTextFont = InTicksTextFont; BroadcastStyleUpdate(); return *this; }
-	FFixedSampleSequenceRulerStyle& SetFontSize(const float InFontSize) { TicksTextFont.Size = InFontSize; BroadcastStyleUpdate(); return *this; }
+	FFixedSampleSequenceRulerStyle& SetTicksTextFont(const FSlateFontInfo& InTicksTextFont) { TicksTextFont = InTicksTextFont; return *this; }
+	FFixedSampleSequenceRulerStyle& SetFontSize(const float InFontSize) { TicksTextFont.Size = InFontSize; return *this; }
 
 	UPROPERTY(EditAnywhere, BlueprintReadWrite, Category = Appearance)
 	float TicksTextOffset;
-	FFixedSampleSequenceRulerStyle& SetTicksTextOffset(const float InTicksTextOffset) { TicksTextOffset = InTicksTextOffset; BroadcastStyleUpdate(); return *this; }
+	FFixedSampleSequenceRulerStyle& SetTicksTextOffset(const float InTicksTextOffset) { TicksTextOffset = InTicksTextOffset; return *this; }
 
 	UPROPERTY(EditAnywhere, BlueprintReadWrite, Category = Appearance)
 	FSlateColor BackgroundColor;
-	FFixedSampleSequenceRulerStyle& SetBackgroundColor(const FSlateColor& InBackgroundColor) { BackgroundColor = InBackgroundColor; BroadcastStyleUpdate(); return *this; }
+	FFixedSampleSequenceRulerStyle& SetBackgroundColor(const FSlateColor& InBackgroundColor) { BackgroundColor = InBackgroundColor; return *this; }
 
 	UPROPERTY(EditAnywhere, BlueprintReadWrite, Category = Appearance)
 	FSlateBrush BackgroundBrush;
-	FFixedSampleSequenceRulerStyle& SetBackgroundBrush(const FSlateBrush& InBackgroundBrush) { BackgroundBrush = InBackgroundBrush; BroadcastStyleUpdate(); return *this; }
+	FFixedSampleSequenceRulerStyle& SetBackgroundBrush(const FSlateBrush& InBackgroundBrush) { BackgroundBrush = InBackgroundBrush; return *this; }
 
 	UPROPERTY(EditAnywhere, BlueprintReadWrite, Category = Appearance)
 	float DesiredWidth;
-	FFixedSampleSequenceRulerStyle& SetDesiredWidth(const float InDesiredWidth) { DesiredWidth = InDesiredWidth; BroadcastStyleUpdate(); return *this; }
+	FFixedSampleSequenceRulerStyle& SetDesiredWidth(const float InDesiredWidth) { DesiredWidth = InDesiredWidth; return *this; }
 
 	UPROPERTY(EditAnywhere, BlueprintReadWrite, Category = Appearance)
 	float DesiredHeight;
-	FFixedSampleSequenceRulerStyle& SetDesiredHeight(const float InDesiredHeight) { DesiredHeight = InDesiredHeight; BroadcastStyleUpdate(); return *this; }
+	FFixedSampleSequenceRulerStyle& SetDesiredHeight(const float InDesiredHeight) { DesiredHeight = InDesiredHeight; return *this; }
 };
 
 /**
  * Represents the appearance of a Sampled Sequence Value Grid Overlay
  */
 USTRUCT(BlueprintType)
-struct AUDIOWIDGETS_API FSampledSequenceValueGridOverlayStyle : public FNotifyingAudioWidgetStyle
+struct AUDIOWIDGETS_API FSampledSequenceValueGridOverlayStyle : public FSlateWidgetStyle
 {
 	GENERATED_USTRUCT_BODY()
 
@@ -358,28 +343,28 @@ struct AUDIOWIDGETS_API FSampledSequenceValueGridOverlayStyle : public FNotifyin
 
 	UPROPERTY(EditAnywhere, BlueprintReadWrite, Category = Appearance)
 	FSlateColor GridColor;
-	FSampledSequenceValueGridOverlayStyle& SetGridColor(const FSlateColor InGridColor) { GridColor = InGridColor; BroadcastStyleUpdate(); return *this; }
+	FSampledSequenceValueGridOverlayStyle& SetGridColor(const FSlateColor InGridColor) { GridColor = InGridColor; return *this; }
 
 	UPROPERTY(EditAnywhere, BlueprintReadWrite, Category = Appearance)
 	float GridThickness;
-	FSampledSequenceValueGridOverlayStyle& SetGridThickness(const float InGridThickness) { GridThickness = InGridThickness; BroadcastStyleUpdate(); return *this; }
+	FSampledSequenceValueGridOverlayStyle& SetGridThickness(const float InGridThickness) { GridThickness = InGridThickness; return *this; }
 
 	UPROPERTY(EditAnywhere, BlueprintReadWrite, Category = Appearance)
 	FSlateColor LabelTextColor;
-	FSampledSequenceValueGridOverlayStyle& SetLabelTextColor(const FSlateColor& InLabelTextColor) { LabelTextColor = InLabelTextColor; BroadcastStyleUpdate(); return *this; }
+	FSampledSequenceValueGridOverlayStyle& SetLabelTextColor(const FSlateColor& InLabelTextColor) { LabelTextColor = InLabelTextColor; return *this; }
 
 	UPROPERTY(EditAnywhere, BlueprintReadWrite, Category = Appearance)
 	FSlateFontInfo LabelTextFont;
-	FSampledSequenceValueGridOverlayStyle& SetLabelTextFont(const FSlateFontInfo& InLabelTextFont) { LabelTextFont = InLabelTextFont; BroadcastStyleUpdate(); return *this; }
-	FSampledSequenceValueGridOverlayStyle& SetLabelTextFontSize(const float InFontSize) { LabelTextFont.Size = InFontSize; BroadcastStyleUpdate(); return *this; }
+	FSampledSequenceValueGridOverlayStyle& SetLabelTextFont(const FSlateFontInfo& InLabelTextFont) { LabelTextFont = InLabelTextFont; return *this; }
+	FSampledSequenceValueGridOverlayStyle& SetLabelTextFontSize(const float InFontSize) { LabelTextFont.Size = InFontSize; return *this; }
 
 	UPROPERTY(EditAnywhere, BlueprintReadWrite, Category = Appearance)
 	float DesiredWidth;
-	FSampledSequenceValueGridOverlayStyle& SetDesiredWidth(const float InDesiredWidth) { DesiredWidth = InDesiredWidth; BroadcastStyleUpdate(); return *this; }
+	FSampledSequenceValueGridOverlayStyle& SetDesiredWidth(const float InDesiredWidth) { DesiredWidth = InDesiredWidth; return *this; }
 
 	UPROPERTY(EditAnywhere, BlueprintReadWrite, Category = Appearance)
 	float DesiredHeight;
-	FSampledSequenceValueGridOverlayStyle& SetDesiredHeight(const float InDesiredHeight) { DesiredHeight = InDesiredHeight; BroadcastStyleUpdate(); return *this; }
+	FSampledSequenceValueGridOverlayStyle& SetDesiredHeight(const float InDesiredHeight) { DesiredHeight = InDesiredHeight; return *this; }
 
 };
 
diff --git a/Engine/Plugins/Runtime/AudioWidgets/Source/AudioWidgets/Public/SFixedSampledSequenceRuler.h b/Engine/Plugins/Runtime/AudioWidgets/Source/AudioWidgets/Public/SFixedSampledSequenceRuler.h
--- a/Engine/Plugins/Runtime/AudioWidgets/Source/AudioWidgets/Public/SFixedSampledSequenceRuler.h
+++ b/Engine/Plugins/Runtime/AudioWidgets/Source/AudioWidgets/Public/SFixedSampledSequenceRuler.h
@@ -33,7 +33,7 @@ public:
 	void UpdateGridMetrics();
 	void UpdateDisplayUnit(const ESampledSequenceDisplayUnit InDisplayUnit);
 	void SetPlayheadPosition(const float InNewPosition);
-	void OnStyleUpdated(const FNotifyingAudioWidgetStyle& UpdatedStyle);
+	void OnStyleUpdated(const FFixedSampleSequenceRulerStyle UpdatedStyle);
 	FReply LaunchContextMenu();
 
 	/** Delegate sent when the user selects a new display unit from the RMB menu*/
@@ -53,8 +53,6 @@ private:
 
 	FFixedSampledSequenceGridMetrics GridMetrics;
 
-	const FFixedSampleSequenceRulerStyle* Style = nullptr;
-
 	FSlateColor BackgroundColor = FLinearColor::Black;
 	FSlateBrush BackgroundBrush;
 	FSlateBrush HandleBrush;
diff --git a/Engine/Plugins/Runtime/AudioWidgets/Source/AudioWidgets/Public/SFixedSampledSequenceViewer.h b/Engine/Plugins/Runtime/AudioWidgets/Source/AudioWidgets/Public/SFixedSampledSequenceViewer.h
--- a/Engine/Plugins/Runtime/AudioWidgets/Source/AudioWidgets/Public/SFixedSampledSequenceViewer.h
+++ b/Engine/Plugins/Runtime/AudioWidgets/Source/AudioWidgets/Public/SFixedSampledSequenceViewer.h
@@ -37,7 +37,7 @@ public:
 	void Construct(const FArguments& InArgs, TArrayView<const float> InSampleData, const uint8 InNumChannels, TSharedRef<IFixedSampledSequenceGridService> InGridService);
 	void UpdateView(TArrayView<const float> InSampleData, const uint8 InNumChannels);
 	
-	void OnStyleUpdated(const FNotifyingAudioWidgetStyle& UpdatedStyle);
+	void OnStyleUpdated(const FSampledSequenceViewerStyle UpdatedStyle);
 
 	void UpdateGridMetrics();
 
@@ -57,7 +57,6 @@ private:
 	TSharedPtr<IFixedSampledSequenceGridService> GridService = nullptr;
 	FFixedSampledSequenceGridMetrics GridMetrics;
 
-	const FSampledSequenceViewerStyle* Style = nullptr;
 	FSlateBrush BackgroundBrush;
 	FSlateColor BackgroundColor = FLinearColor(0.02f, 0.02f, 0.02f, 1.f);
 	FSlateColor SequenceColor = FLinearColor::White;
diff --git a/Engine/Plugins/Runtime/AudioWidgets/Source/AudioWidgets/Public/SFixedSampledSequenceViewerPanel.h b/Engine/Plugins/Runtime/AudioWidgets/Source/AudioWidgets/Public/SFixedSampledSequenceViewerPanel.h
--- a/Engine/Plugins/Runtime/AudioWidgets/Source/AudioWidgets/Public/SFixedSampledSequenceViewerPanel.h
+++ b/Engine/Plugins/Runtime/AudioWidgets/Source/AudioWidgets/Public/SFixedSampledSequenceViewerPanel.h
@@ -88,6 +88,11 @@ public:
 	virtual void ReceiveSequenceView(const FFixedSampledSequenceView InData, const uint32 FirstSampleIndex) override;
 	
 	void SetPlayheadRatio(const float InRatio);
+
+	void UpdateSequenceViewerStyle(const FSampledSequenceViewerStyle UpdatedStyle);
+	void UpdateRulerStyle(const FFixedSampleSequenceRulerStyle UpdatedStyle);
+	void UpdatePlayheadOverlayStyle(const FPlayheadOverlayStyle UpdatedStyle);
+	void UpdateValueGridOverlayStyle(const FSampledSequenceValueGridOverlayStyle UpdatedStyle);
 	
 private:
 	virtual void CreateLayout();
@@ -110,7 +115,6 @@ private:
 	void UpdateSamplesRulerDisplayUnit(const ESampledSequenceDisplayUnit InDisplayUnit);
 	void UpdatePlayheadPosition(const float PaintedWidth);
 
-
 	float CachedPixelWidth = 0.f;
 
 	ESampledSequenceDisplayUnit SamplesRulerDisplayUnit;
diff --git a/Engine/Plugins/Runtime/AudioWidgets/Source/AudioWidgets/Public/SPlayheadOverlay.h b/Engine/Plugins/Runtime/AudioWidgets/Source/AudioWidgets/Public/SPlayheadOverlay.h
--- a/Engine/Plugins/Runtime/AudioWidgets/Source/AudioWidgets/Public/SPlayheadOverlay.h
+++ b/Engine/Plugins/Runtime/AudioWidgets/Source/AudioWidgets/Public/SPlayheadOverlay.h
@@ -20,7 +20,7 @@ public:
 
 	void Construct(const FArguments& InArgs);
 	void SetPlayheadPosition(const float InNewPosition);
-	void OnStyleUpdated(const FNotifyingAudioWidgetStyle& UpdatedStyle);
+	void OnStyleUpdated(const FPlayheadOverlayStyle UpdatedStyle);
 
 private:
 	virtual int32 OnPaint(const FPaintArgs& Args, const FGeometry& AllottedGeometry, const FSlateRect& MyCullingRect, FSlateWindowElementList& OutDrawElements, int32 LayerId, const FWidgetStyle& InWidgetStyle, bool bParentEnabled) const override;
diff --git a/Engine/Plugins/Runtime/AudioWidgets/Source/AudioWidgets/Public/SSampledSequenceValueGridOverlay.h b/Engine/Plugins/Runtime/AudioWidgets/Source/AudioWidgets/Public/SSampledSequenceValueGridOverlay.h
--- a/Engine/Plugins/Runtime/AudioWidgets/Source/AudioWidgets/Public/SSampledSequenceValueGridOverlay.h
+++ b/Engine/Plugins/Runtime/AudioWidgets/Source/AudioWidgets/Public/SSampledSequenceValueGridOverlay.h
@@ -53,7 +53,7 @@ public:
 	virtual int32 OnPaint(const FPaintArgs& Args, const FGeometry& AllottedGeometry, const FSlateRect& MyCullingRect, FSlateWindowElementList& OutDrawElements, int32 LayerId, const FWidgetStyle& InWidgetStyle, bool bParentEnabled) const override;
 	virtual FVector2D ComputeDesiredSize(float) const override;
 	virtual void Tick(const FGeometry& AllottedGeometry, const double InCurrentTime, const float InDeltaTime) override;
-	void OnStyleUpdated(const FNotifyingAudioWidgetStyle& UpdatedStyle);
+	void OnStyleUpdated(const FSampledSequenceValueGridOverlayStyle UpdatedStyle);
 
 	void SetMaxDivisionParameter(const uint32 InDivisionParameter);
 	void SetHideLabels(const bool InHideLabels);
