[AudioWidgets] Export FFixedSampledSequenceGridData to be used in SweetJTech oscilloscope
#trivial

diff --git a/Engine/Plugins/Runtime/AudioWidgets/Source/AudioWidgets/Public/FixedSampledSequenceGridData.h b/Engine/Plugins/Runtime/AudioWidgets/Source/AudioWidgets/Public/FixedSampledSequenceGridData.h
--- a/Engine/Plugins/Runtime/AudioWidgets/Source/AudioWidgets/Public/FixedSampledSequenceGridData.h
+++ b/Engine/Plugins/Runtime/AudioWidgets/Source/AudioWidgets/Public/FixedSampledSequenceGridData.h
@@ -12,7 +12,7 @@ DECLARE_MULTICAST_DELEGATE(FOnWaveformEditorGridUpdated)
 
 struct FSlateFontInfo;
 
-class FFixedSampledSequenceGridData : public IFixedSampledSequenceGridService
+class AUDIOWIDGETS_API FFixedSampledSequenceGridData : public IFixedSampledSequenceGridService
 {
 public:
 	explicit FFixedSampledSequenceGridData(const uint32 InTotalFrames, const uint32 InSampleRateHz, const float InGridSizePixels = 1.f, const FSlateFontInfo* InTicksTimeFont = nullptr);
