expose AudioWidgetsDrawigUtils::FDimensionSlot for oscilloscope use
#trival

diff --git a/Engine/Plugins/Runtime/AudioWidgets/Source/AudioWidgets/Public/SampledSequenceDrawingUtils.h b/Engine/Plugins/Runtime/AudioWidgets/Source/AudioWidgets/Public/SampledSequenceDrawingUtils.h
--- a/Engine/Plugins/Runtime/AudioWidgets/Source/AudioWidgets/Public/SampledSequenceDrawingUtils.h
+++ b/Engine/Plugins/Runtime/AudioWidgets/Source/AudioWidgets/Public/SampledSequenceDrawingUtils.h
@@ -125,7 +125,7 @@ namespace SampledSequenceDrawingUtils
 	/**
 	* Represents a dimensional slot in which samples, sample bins or a grid can be drawn
 	*/
-	struct FDimensionSlot
+	struct AUDIOWIDGETS_API FDimensionSlot
 	{
 		explicit FDimensionSlot(const uint16 DimensionToDraw, const uint16 TotalNumDimensions, const FGeometry& InAllottedGeometry, const FSampledSequenceDrawingParams& Params);
 
