[AudioWidgets] SSampledSequenceValueGridOverlay: Added SetLabelGenerator method so we can do live changes to the grid units + SFixedSampledSequenceViewer: Added HideGrid method to hide/show the X axis grid when doing live changes.

#rnx

diff --git a/Engine/Plugins/Runtime/AudioWidgets/Source/AudioWidgets/Private/SFixedSampledSequenceViewer.cpp b/Engine/Plugins/Runtime/AudioWidgets/Source/AudioWidgets/Private/SFixedSampledSequenceViewer.cpp
--- a/Engine/Plugins/Runtime/AudioWidgets/Source/AudioWidgets/Private/SFixedSampledSequenceViewer.cpp
+++ b/Engine/Plugins/Runtime/AudioWidgets/Source/AudioWidgets/Private/SFixedSampledSequenceViewer.cpp
@@ -271,6 +271,11 @@ void SFixedSampledSequenceViewer::UpdateGridMetrics()
 	GridMetrics = GridService->GetGridMetrics();
 }
 
+void SFixedSampledSequenceViewer::SetHideGrid(const bool InHideGrid)
+{
+	bHideGrid = InHideGrid;
+}
+
 void SFixedSampledSequenceViewer::OnStyleUpdated(const FSampledSequenceViewerStyle UpdatedStyle)
 {
 	SequenceColor = UpdatedStyle.SequenceColor;
diff --git a/Engine/Plugins/Runtime/AudioWidgets/Source/AudioWidgets/Private/SSampledSequenceValueGridOverlay.cpp b/Engine/Plugins/Runtime/AudioWidgets/Source/AudioWidgets/Private/SSampledSequenceValueGridOverlay.cpp
--- a/Engine/Plugins/Runtime/AudioWidgets/Source/AudioWidgets/Private/SSampledSequenceValueGridOverlay.cpp
+++ b/Engine/Plugins/Runtime/AudioWidgets/Source/AudioWidgets/Private/SSampledSequenceValueGridOverlay.cpp
@@ -139,6 +139,12 @@ void SSampledSequenceValueGridOverlay::OnStyleUpdated(const FSampledSequenceValu
 	ForceRedraw();
 }
 
+void SSampledSequenceValueGridOverlay::SetLabelGenerator(TFunction<FText(const double)> InLabelGenerator)
+{
+	OnValueGridLabel = InLabelGenerator;
+	ForceRedraw();
+}
+
 void SSampledSequenceValueGridOverlay::SetMaxDivisionParameter(const uint32 InDivisionParameter)
 {
 	MaxDivisionParameter = InDivisionParameter;
diff --git a/Engine/Plugins/Runtime/AudioWidgets/Source/AudioWidgets/Public/SFixedSampledSequenceViewer.h b/Engine/Plugins/Runtime/AudioWidgets/Source/AudioWidgets/Public/SFixedSampledSequenceViewer.h
--- a/Engine/Plugins/Runtime/AudioWidgets/Source/AudioWidgets/Public/SFixedSampledSequenceViewer.h
+++ b/Engine/Plugins/Runtime/AudioWidgets/Source/AudioWidgets/Public/SFixedSampledSequenceViewer.h
@@ -41,6 +41,8 @@ public:
 
 	void UpdateGridMetrics();
 
+	void SetHideGrid(const bool InHideGrid);
+
 private:
 	virtual int32 OnPaint(const FPaintArgs& Args, const FGeometry& AllottedGeometry, const FSlateRect& MyCullingRect, FSlateWindowElementList& OutDrawElements, int32 LayerId, const FWidgetStyle& InWidgetStyle, bool bParentEnabled) const override;
 	virtual FVector2D ComputeDesiredSize(float) const override;
diff --git a/Engine/Plugins/Runtime/AudioWidgets/Source/AudioWidgets/Public/SSampledSequenceValueGridOverlay.h b/Engine/Plugins/Runtime/AudioWidgets/Source/AudioWidgets/Public/SSampledSequenceValueGridOverlay.h
--- a/Engine/Plugins/Runtime/AudioWidgets/Source/AudioWidgets/Public/SSampledSequenceValueGridOverlay.h
+++ b/Engine/Plugins/Runtime/AudioWidgets/Source/AudioWidgets/Public/SSampledSequenceValueGridOverlay.h
@@ -55,6 +55,7 @@ public:
 	virtual void Tick(const FGeometry& AllottedGeometry, const double InCurrentTime, const float InDeltaTime) override;
 	void OnStyleUpdated(const FSampledSequenceValueGridOverlayStyle UpdatedStyle);
 
+	void SetLabelGenerator(TFunction<FText(const double)> InLabelGenerator);
 	void SetMaxDivisionParameter(const uint32 InDivisionParameter);
 	void SetHideLabels(const bool InHideLabels);
 	void SetHideGrid(const bool InHideGrid);
