[AudioWidgets] Moved Oscilloscope and Vectorscope from AudioVisualization plugin to AudioWidgets plugin

diff --git a/Engine/Plugins/Runtime/AudioWidgets/Source/AudioWidgets/AudioWidgets.Build.cs b/Engine/Plugins/Runtime/AudioWidgets/Source/AudioWidgets/AudioWidgets.Build.cs
--- a/Engine/Plugins/Runtime/AudioWidgets/Source/AudioWidgets/AudioWidgets.Build.cs
+++ b/Engine/Plugins/Runtime/AudioWidgets/Source/AudioWidgets/AudioWidgets.Build.cs
@@ -23,13 +23,26 @@ public class AudioWidgets : ModuleRules
 		(
 			new string[]
 			{
+				"AudioMixer",
 				"CoreUObject",
 				"Engine",
+				"InputCore",
 				"Slate",
 				"SlateCore",
 				"AdvancedWidgets",
 				"SignalProcessing"
 			}
 		);
+
+		if (Target.bBuildEditor)
+		{
+			PrivateDependencyModuleNames.AddRange
+			(
+				new string[]
+				{
+					"UnrealEd",
+				}
+			);
+		}
 	}
 }
diff --git a/Engine/Plugins/Runtime/AudioWidgets/Source/AudioWidgets/Private/AudioOscilloscope.cpp b/Engine/Plugins/Runtime/AudioWidgets/Source/AudioWidgets/Private/AudioOscilloscope.cpp
new file mode 100644
--- /dev/null
+++ b/Engine/Plugins/Runtime/AudioWidgets/Source/AudioWidgets/Private/AudioOscilloscope.cpp
@@ -0,0 +1,64 @@
+// Copyright Epic Games, Inc. All Rights Reserved.
+
+#include "AudioOscilloscope.h"
+
+#include "Engine/World.h"
+
+namespace AudioWidgets
+{
+	FAudioOscilloscope::FAudioOscilloscope(UWorld* InWorld, 
+		const uint32 InNumChannels, 
+		const float InTimeWindowMs, 
+		const float InMaxTimeWindowMs, 
+		const float InAnalysisPeriodMs, 
+		const EAudioPanelLayoutType InPanelLayoutType)
+		: OscilloscopePanelStyle(FAudioOscilloscopePanelStyle::GetDefault())
+	{
+		// Init audio bus
+		AudioBus = TStrongObjectPtr(NewObject<UAudioBus>());
+		AudioBus->AudioBusChannels = AudioBusUtils::ConvertIntToEAudioBusChannels(InNumChannels);
+
+		// Init data provider
+		constexpr uint32 NumChannelsToProvide = 1;
+		AudioSamplesDataProvider = MakeShared<FWaveformAudioSamplesDataProvider>(InWorld, AudioBus.Get(), NumChannelsToProvide, InTimeWindowMs, InMaxTimeWindowMs, InAnalysisPeriodMs);
+
+		// Init widget 
+		const FFixedSampledSequenceView SequenceView = AudioSamplesDataProvider->GetDataView();
+
+		OscilloscopePanelWidget = SNew(SAudioOscilloscopePanelWidget, SequenceView, InNumChannels)
+		.PanelLayoutType(InPanelLayoutType)
+		.PanelStyle(&OscilloscopePanelStyle);
+
+		// Interconnect data provider and widget
+		AudioSamplesDataProvider->OnDataViewGenerated.AddSP(OscilloscopePanelWidget.Get(), &SAudioOscilloscopePanelWidget::ReceiveSequenceView);
+
+		if (InPanelLayoutType == EAudioPanelLayoutType::Advanced)
+		{
+			OscilloscopePanelWidget->OnSelectedChannelChanged.AddSP(AudioSamplesDataProvider.Get(),  &FWaveformAudioSamplesDataProvider::SetChannelToAnalyze);
+			OscilloscopePanelWidget->OnTriggerModeChanged.AddSP(AudioSamplesDataProvider.Get(),      &FWaveformAudioSamplesDataProvider::SetTriggerMode);
+			OscilloscopePanelWidget->OnTriggerThresholdChanged.AddSP(AudioSamplesDataProvider.Get(), &FWaveformAudioSamplesDataProvider::SetTriggerThreshold);
+			OscilloscopePanelWidget->OnTimeWindowValueChanged.AddSP(AudioSamplesDataProvider.Get(),  &FWaveformAudioSamplesDataProvider::SetTimeWindow);
+			OscilloscopePanelWidget->OnAnalysisPeriodChanged.AddSP(AudioSamplesDataProvider.Get(),   &FWaveformAudioSamplesDataProvider::SetAnalysisPeriod);
+		}
+	}
+
+	void FAudioOscilloscope::StartProcessing()
+	{
+		AudioSamplesDataProvider->StartProcessing();
+	}
+
+	void FAudioOscilloscope::StopProcessing()
+	{
+		AudioSamplesDataProvider->StopProcessing();
+	}
+
+	UAudioBus* FAudioOscilloscope::GetAudioBus() const
+	{
+		return AudioBus.Get();
+	}
+
+	TSharedRef<SWidget> FAudioOscilloscope::GetPanelWidget() const
+	{
+		return OscilloscopePanelWidget.ToSharedRef();
+	}
+}
diff --git a/Engine/Plugins/Runtime/AudioWidgets/Source/AudioWidgets/Private/AudioOscilloscopeUMG.cpp b/Engine/Plugins/Runtime/AudioWidgets/Source/AudioWidgets/Private/AudioOscilloscopeUMG.cpp
new file mode 100644
--- /dev/null
+++ b/Engine/Plugins/Runtime/AudioWidgets/Source/AudioWidgets/Private/AudioOscilloscopeUMG.cpp
@@ -0,0 +1,168 @@
+// Copyright Epic Games, Inc. All Rights Reserved.
+
+#include "AudioOscilloscopeUMG.h"
+
+#include "AudioMixerDevice.h"
+#include "Engine/Engine.h"
+#include "Sound/AudioBus.h"
+#include "WaveformAudioSamplesDataProvider.h"
+
+#include UE_INLINE_GENERATED_CPP_BY_NAME(AudioOscilloscopeUMG)
+
+#define LOCTEXT_NAMESPACE "AudioOscilloscopeUMG"
+
+UAudioOscilloscope::UAudioOscilloscope(const FObjectInitializer& ObjectInitializer)
+	: Super(ObjectInitializer)
+	, OscilloscopeStyle(FAudioOscilloscopePanelStyle::GetDefault())
+{
+#if WITH_EDITORONLY_DATA
+	AccessibleBehavior       = ESlateAccessibleBehavior::NotAccessible;
+	bCanChildrenBeAccessible = false;
+#endif
+
+	DummyAudioSamples.Init(0.0f, DummyMaxNumSamples);
+	DummyDataView = FFixedSampledSequenceView{ MakeArrayView(DummyAudioSamples.GetData(), DummyAudioSamples.Num()), DummyNumChannels, DummySampleRate };
+}
+
+void UAudioOscilloscope::CreateDummyOscilloscopeWidget()
+{
+	OscilloscopePanelWidget = SNew(SAudioOscilloscopePanelWidget, DummyDataView, 1)
+	.PanelLayoutType(PanelLayoutType)
+	.PanelStyle(&OscilloscopeStyle);
+}
+
+void UAudioOscilloscope::CreateDataProvider()
+{
+	constexpr uint32 NumChannelsToProvide = 1;
+	constexpr float  MaxTimeWindowMs      = 5000.0f; // TODO alex.perez: should we expose this as a UPROPERTY?
+	AudioSamplesDataProvider = MakeShared<AudioWidgets::FWaveformAudioSamplesDataProvider>(GetWorld(), AudioBus, NumChannelsToProvide, TimeWindowMs, MaxTimeWindowMs, AnalysisPeriodMs);
+}
+
+void UAudioOscilloscope::CreateOscilloscopeWidget()
+{
+	check(AudioSamplesDataProvider);
+
+	using namespace AudioWidgets;
+
+	const FFixedSampledSequenceView SequenceView = AudioSamplesDataProvider->GetDataView();
+
+	OscilloscopePanelWidget = SNew(SAudioOscilloscopePanelWidget, SequenceView, AudioSamplesDataProvider->GetNumChannels())
+	.PanelLayoutType(PanelLayoutType)
+	.PanelStyle(&OscilloscopeStyle);
+
+	// Interconnect data provider and widget
+	AudioSamplesDataProvider->OnDataViewGenerated.AddSP(OscilloscopePanelWidget.Get(), &SAudioOscilloscopePanelWidget::ReceiveSequenceView);
+
+	if (PanelLayoutType == EAudioPanelLayoutType::Advanced)
+	{
+		OscilloscopePanelWidget->OnSelectedChannelChanged.AddSP(AudioSamplesDataProvider.Get(),  &FWaveformAudioSamplesDataProvider::SetChannelToAnalyze);
+		OscilloscopePanelWidget->OnTriggerModeChanged.AddSP(AudioSamplesDataProvider.Get(),      &FWaveformAudioSamplesDataProvider::SetTriggerMode);
+		OscilloscopePanelWidget->OnTriggerThresholdChanged.AddSP(AudioSamplesDataProvider.Get(), &FWaveformAudioSamplesDataProvider::SetTriggerThreshold);
+		OscilloscopePanelWidget->OnTimeWindowValueChanged.AddSP(AudioSamplesDataProvider.Get(),  &FWaveformAudioSamplesDataProvider::SetTimeWindow);
+		OscilloscopePanelWidget->OnAnalysisPeriodChanged.AddSP(AudioSamplesDataProvider.Get(),   &FWaveformAudioSamplesDataProvider::SetAnalysisPeriod);
+	}
+}
+
+TSharedRef<SWidget> UAudioOscilloscope::RebuildWidget()
+{
+	if (!AudioBus)
+	{
+		CreateDummyOscilloscopeWidget();
+	}
+	else
+	{
+		CreateDataProvider();
+		CreateOscilloscopeWidget();
+	}
+
+	return OscilloscopePanelWidget.ToSharedRef();
+}
+
+void UAudioOscilloscope::SynchronizeProperties()
+{
+	Super::SynchronizeProperties();
+
+	if (!AudioBus)
+	{
+		if (AudioSamplesDataProvider.IsValid())
+		{
+			AudioSamplesDataProvider.Reset();
+			CreateDummyOscilloscopeWidget();
+		}
+	}
+	else
+	{
+		if (!AudioSamplesDataProvider.IsValid() || AudioBus != AudioSamplesDataProvider->GetAudioBus())
+		{
+			CreateDataProvider();
+			CreateOscilloscopeWidget();
+		}
+	}
+
+	if (AudioSamplesDataProvider.IsValid())
+	{
+		AudioSamplesDataProvider->SetTimeWindow(TimeWindowMs);
+		AudioSamplesDataProvider->SetAnalysisPeriod(AnalysisPeriodMs);
+		AudioSamplesDataProvider->RequestSequenceView(TRange<double>::Inclusive(0, 1));
+	}
+
+	if (OscilloscopePanelWidget.IsValid())
+	{
+		if (PanelLayoutType != OscilloscopePanelWidget->GetPanelLayoutType() && AudioSamplesDataProvider.IsValid())
+		{
+			CreateOscilloscopeWidget();
+		}
+
+		if (!AudioBus)
+		{
+			DummyDataView = FFixedSampledSequenceView{ MakeArrayView(DummyAudioSamples.GetData(), static_cast<int32>((TimeWindowMs / 1000.0f) * DummySampleRate)), DummyNumChannels, DummySampleRate};
+			OscilloscopePanelWidget->ReceiveSequenceView(DummyDataView, 0);
+		}
+
+		OscilloscopePanelWidget->UpdateSequenceRulerStyle(OscilloscopeStyle.TimeRulerStyle);
+		OscilloscopePanelWidget->UpdateValueGridOverlayStyle(OscilloscopeStyle.ValueGridStyle);
+		OscilloscopePanelWidget->UpdateSequenceViewerStyle(OscilloscopeStyle.WaveViewerStyle);
+
+		OscilloscopePanelWidget->SetXAxisGridVisibility(bShowTimeGrid);
+		OscilloscopePanelWidget->SetSequenceRulerDisplayUnit(TimeGridLabelsUnit);
+
+		OscilloscopePanelWidget->SetYAxisGridVisibility(bShowAmplitudeGrid);
+		OscilloscopePanelWidget->SetYAxisLabelsVisibility(bShowAmplitudeLabels);
+		OscilloscopePanelWidget->SetValueGridOverlayDisplayUnit(AmplitudeGridLabelsUnit);
+
+		OscilloscopePanelWidget->SetTriggerThreshold(TriggerThreshold);
+		OscilloscopePanelWidget->SetTriggerThresholdVisibility(bShowTriggerThresholdLine);
+	}
+}
+
+void UAudioOscilloscope::ReleaseSlateResources(bool bReleaseChildren)
+{
+	Super::ReleaseSlateResources(bReleaseChildren);
+	
+	OscilloscopePanelWidget.Reset();
+}
+
+#if WITH_EDITOR
+const FText UAudioOscilloscope::GetPaletteCategory()
+{
+	return LOCTEXT("Audio", "Audio");
+}
+#endif
+
+void UAudioOscilloscope::StartProcessing()
+{
+	if (AudioSamplesDataProvider.IsValid())
+	{
+		AudioSamplesDataProvider->StartProcessing();
+	}
+}
+
+void UAudioOscilloscope::StopProcessing()
+{
+	if (AudioSamplesDataProvider.IsValid())
+	{
+		AudioSamplesDataProvider->StopProcessing();
+	}
+}
+
+#undef LOCTEXT_NAMESPACE
diff --git a/Engine/Plugins/Runtime/AudioWidgets/Source/AudioWidgets/Private/AudioVectorscope.cpp b/Engine/Plugins/Runtime/AudioWidgets/Source/AudioWidgets/Private/AudioVectorscope.cpp
new file mode 100644
--- /dev/null
+++ b/Engine/Plugins/Runtime/AudioWidgets/Source/AudioWidgets/Private/AudioVectorscope.cpp
@@ -0,0 +1,60 @@
+// Copyright Epic Games, Inc. All Rights Reserved.
+
+#include "AudioVectorscope.h"
+
+#include "Engine/World.h"
+#include "SAudioVectorscopePanelWidget.h"
+
+namespace AudioWidgets
+{
+	FAudioVectorscope::FAudioVectorscope(UWorld* InWorld, 
+		const uint32 InNumChannels, 
+		const float InTimeWindowMs, 
+		const float InMaxTimeWindowMs, 
+		const float InAnalysisPeriodMs, 
+		const EAudioPanelLayoutType InPanelLayoutType)
+		: VectorscopePanelStyle(FAudioVectorscopePanelStyle::GetDefault())
+	{
+		// Init audio bus
+		AudioBus = TStrongObjectPtr(NewObject<UAudioBus>());
+		AudioBus->AudioBusChannels = AudioBusUtils::ConvertIntToEAudioBusChannels(InNumChannels);
+
+		// Init data provider
+		AudioSamplesDataProvider = MakeShared<FWaveformAudioSamplesDataProvider>(InWorld, AudioBus.Get(), AudioBus->GetNumChannels(), InTimeWindowMs, InMaxTimeWindowMs, InAnalysisPeriodMs);
+
+		// Init widget 
+		const FFixedSampledSequenceView SequenceView = AudioSamplesDataProvider->GetDataView();
+
+		VectorscopePanelWidget = SNew(SAudioVectorscopePanelWidget, SequenceView)
+		.PanelLayoutType(InPanelLayoutType)
+		.PanelStyle(&VectorscopePanelStyle);
+
+		// Interconnect data provider and widget
+		AudioSamplesDataProvider->OnDataViewGenerated.AddSP(VectorscopePanelWidget.Get(), &SAudioVectorscopePanelWidget::ReceiveSequenceView);
+
+		if (InPanelLayoutType == EAudioPanelLayoutType::Advanced)
+		{
+			VectorscopePanelWidget->OnTimeWindowValueChanged.AddSP(AudioSamplesDataProvider.Get(), &FWaveformAudioSamplesDataProvider::SetTimeWindow);
+		}
+	}
+
+	void FAudioVectorscope::StartProcessing()
+	{
+		AudioSamplesDataProvider->StartProcessing();
+	}
+
+	void FAudioVectorscope::StopProcessing()
+	{
+		AudioSamplesDataProvider->StopProcessing();
+	}
+
+	UAudioBus* FAudioVectorscope::GetAudioBus() const
+	{
+		return AudioBus.Get();
+	}
+
+	TSharedRef<SWidget> FAudioVectorscope::GetPanelWidget() const
+	{
+		return VectorscopePanelWidget.ToSharedRef();
+	}
+} // namespace AudioWidgets
diff --git a/Engine/Plugins/Runtime/AudioWidgets/Source/AudioWidgets/Private/AudioVectorscopeUMG.cpp b/Engine/Plugins/Runtime/AudioWidgets/Source/AudioWidgets/Private/AudioVectorscopeUMG.cpp
new file mode 100644
--- /dev/null
+++ b/Engine/Plugins/Runtime/AudioWidgets/Source/AudioWidgets/Private/AudioVectorscopeUMG.cpp
@@ -0,0 +1,143 @@
+// Copyright Epic Games, Inc. All Rights Reserved.
+
+#include "AudioVectorscopeUMG.h"
+
+#include "AudioMixerDevice.h"
+#include "SAudioVectorscopePanelWidget.h"
+#include "WaveformAudioSamplesDataProvider.h"
+
+#include UE_INLINE_GENERATED_CPP_BY_NAME(AudioVectorscopeUMG)
+
+#define LOCTEXT_NAMESPACE "AudioVectorscopeUMG"
+
+UAudioVectorscope::UAudioVectorscope(const FObjectInitializer& ObjectInitializer)
+	: Super(ObjectInitializer)
+	, VectorscopeStyle(FAudioVectorscopePanelStyle::GetDefault())
+{
+#if WITH_EDITORONLY_DATA
+	AccessibleBehavior       = ESlateAccessibleBehavior::NotAccessible;
+	bCanChildrenBeAccessible = false;
+#endif
+
+	DummyAudioSamples.Init(0.0f, 100);
+	DummyDataView = FFixedSampledSequenceView{ MakeArrayView(DummyAudioSamples.GetData(), DummyAudioSamples.Num()), 2, 48000 };
+}
+
+void UAudioVectorscope::CreateDummyVectorscopeWidget()
+{
+	VectorscopePanelWidget = SNew(SAudioVectorscopePanelWidget, DummyDataView)
+	.PanelLayoutType(PanelLayoutType)
+	.PanelStyle(&VectorscopeStyle);
+}
+
+void UAudioVectorscope::CreateDataProvider()
+{
+	constexpr float MaxDisplayPersistenceMs = 500.0f; // TODO alex.perez: should we expose this as a UPROPERTY?
+	AudioSamplesDataProvider = MakeShared<AudioWidgets::FWaveformAudioSamplesDataProvider>(GetWorld(), AudioBus, AudioBus->GetNumChannels(), DisplayPersistenceMs, MaxDisplayPersistenceMs, AnalysisPeriodMs);
+}
+
+void UAudioVectorscope::CreateVectorscopeWidget()
+{
+	check(AudioSamplesDataProvider);
+
+	using namespace AudioWidgets;
+
+	const FFixedSampledSequenceView SequenceView = AudioSamplesDataProvider->GetDataView();
+
+	VectorscopePanelWidget = SNew(SAudioVectorscopePanelWidget, SequenceView)
+	.PanelLayoutType(PanelLayoutType)
+	.PanelStyle(&VectorscopeStyle);
+
+	// Interconnect data provider and widget
+	AudioSamplesDataProvider->OnDataViewGenerated.AddSP(VectorscopePanelWidget.Get(), &SAudioVectorscopePanelWidget::ReceiveSequenceView);
+
+	if (PanelLayoutType == EAudioPanelLayoutType::Advanced)
+	{
+		VectorscopePanelWidget->OnTimeWindowValueChanged.AddSP(AudioSamplesDataProvider.Get(), &FWaveformAudioSamplesDataProvider::SetTimeWindow);
+	}
+}
+
+TSharedRef<SWidget> UAudioVectorscope::RebuildWidget()
+{
+	if (!AudioBus)
+	{
+		CreateDummyVectorscopeWidget();
+	}
+	else
+	{
+		CreateDataProvider();
+		CreateVectorscopeWidget();
+	}
+
+	return VectorscopePanelWidget.ToSharedRef();
+}
+
+void UAudioVectorscope::SynchronizeProperties()
+{
+	Super::SynchronizeProperties();
+
+	if (!AudioBus)
+	{
+		if (AudioSamplesDataProvider.IsValid())
+		{
+			AudioSamplesDataProvider.Reset();
+			CreateDummyVectorscopeWidget();
+		}
+	}
+	else
+	{
+		if (!AudioSamplesDataProvider.IsValid() || AudioBus != AudioSamplesDataProvider->GetAudioBus())
+		{
+			CreateDataProvider();
+			CreateVectorscopeWidget();
+		}
+	}
+
+	if (VectorscopePanelWidget.IsValid())
+	{
+		if (PanelLayoutType != VectorscopePanelWidget->GetPanelLayoutType() && AudioSamplesDataProvider.IsValid())
+		{
+			CreateVectorscopeWidget();
+		}
+
+		VectorscopePanelWidget->UpdateValueGridOverlayStyle(VectorscopeStyle.ValueGridStyle);
+		VectorscopePanelWidget->SetValueGridOverlayMaxNumDivisions(GridDivisions);
+
+		VectorscopePanelWidget->UpdateSequenceVectorViewerStyle(VectorscopeStyle.VectorViewerStyle);
+		VectorscopePanelWidget->SetVectorViewerScaleFactor(Scale);
+
+		VectorscopePanelWidget->SetGridVisibility(bShowGrid);
+	}
+}
+
+void UAudioVectorscope::ReleaseSlateResources(bool bReleaseChildren)
+{
+	Super::ReleaseSlateResources(bReleaseChildren);
+	
+	VectorscopePanelWidget.Reset();
+}
+
+#if WITH_EDITOR
+const FText UAudioVectorscope::GetPaletteCategory()
+{
+	return LOCTEXT("Audio", "Audio");
+}
+#endif
+
+void UAudioVectorscope::StartProcessing()
+{
+	if (AudioSamplesDataProvider.IsValid())
+	{
+		AudioSamplesDataProvider->StartProcessing();
+	}
+}
+
+void UAudioVectorscope::StopProcessing()
+{
+	if (AudioSamplesDataProvider.IsValid())
+	{
+		AudioSamplesDataProvider->StopProcessing();
+	}
+}
+
+#undef LOCTEXT_NAMESPACE
diff --git a/Engine/Plugins/Runtime/AudioWidgets/Source/AudioWidgets/Private/AudioWidgetsLog.cpp b/Engine/Plugins/Runtime/AudioWidgets/Source/AudioWidgets/Private/AudioWidgetsLog.cpp
new file mode 100644
--- /dev/null
+++ b/Engine/Plugins/Runtime/AudioWidgets/Source/AudioWidgets/Private/AudioWidgetsLog.cpp
@@ -0,0 +1,5 @@
+// Copyright Epic Games, Inc. All Rights Reserved.
+
+#include "AudioWidgetsLog.h"
+
+DEFINE_LOG_CATEGORY(LogAudioWidgets);
diff --git a/Engine/Plugins/Runtime/AudioWidgets/Source/AudioWidgets/Private/SAudioOscilloscopePanelWidget.cpp b/Engine/Plugins/Runtime/AudioWidgets/Source/AudioWidgets/Private/SAudioOscilloscopePanelWidget.cpp
new file mode 100644
--- /dev/null
+++ b/Engine/Plugins/Runtime/AudioWidgets/Source/AudioWidgets/Private/SAudioOscilloscopePanelWidget.cpp
@@ -0,0 +1,737 @@
+// Copyright Epic Games, Inc. All Rights Reserved.
+
+#include "SAudioOscilloscopePanelWidget.h"
+
+#include "AudioWidgetsLog.h"
+#include "DSP/Dsp.h"
+#include "FixedSampledSequenceGridData.h"
+#include "SAudioRadialSlider.h"
+#include "SFixedSampledSequenceRuler.h"
+#include "SFixedSampledSequenceViewer.h"
+#include "Widgets/Layout/SBorder.h"
+#include "Widgets/SOverlay.h"
+
+#if WITH_EDITOR
+#include "Editor.h"
+#endif
+
+#define LOCTEXT_NAMESPACE "SOscilloscopePanelWidget"
+
+namespace AudioOscilloscopePanelWidgetPrivate
+{
+	TFunction<FText(const double)> YAxisValueGeneratorLinear = [](const double InLabelValue)
+	{
+		const float LinearAmplitude = 1.0f - 2.0f * InLabelValue;
+
+		const FNumberFormattingOptions Formatting = FNumberFormattingOptions().SetMaximumFractionalDigits(2);
+
+		return FText::AsNumber(LinearAmplitude, &Formatting);
+	};
+
+	TFunction<FText(const double)> YAxisValueGeneratorDb = [](const double InLabelValue)
+	{
+		constexpr float MinDecibelNumericLimit = -160.0f;
+
+		const float AbsAmplitude  = FMath::Abs(InLabelValue * 2.0f - 1.0f);
+		const float GridLineValue = Audio::ConvertToDecibels(AbsAmplitude);
+
+		const FNumberFormattingOptions Formatting = FNumberFormattingOptions().SetMaximumFractionalDigits(1);
+
+		const FText GridLineValueText = (GridLineValue <= MinDecibelNumericLimit) ? FText::FromString("-inf") : FText::AsNumber(GridLineValue, &Formatting);
+
+		return GridLineValueText;
+	};
+
+	const TFunction<FText(const double)>& GetValueGridOverlayLabelGenerator(const EYAxisLabelsUnit InValueGridOverlayDisplayUnit)
+	{
+		switch (InValueGridOverlayDisplayUnit)
+		{
+		default:
+		case EYAxisLabelsUnit::Linear:
+			return AudioOscilloscopePanelWidgetPrivate::YAxisValueGeneratorLinear;
+			break;
+		case EYAxisLabelsUnit::Db:
+			return AudioOscilloscopePanelWidgetPrivate::YAxisValueGeneratorDb;
+			break;
+		}
+	};
+
+	ESampledSequenceDisplayUnit ConvertEXAxisLabelsUnitToESampledSequenceDisplayUnit(EXAxisLabelsUnit InXAxisLabelsUnit)
+	{
+		switch (InXAxisLabelsUnit)
+		{
+			default:
+			case EXAxisLabelsUnit::Samples:
+				return ESampledSequenceDisplayUnit::Samples;
+			case EXAxisLabelsUnit::Seconds:
+				return ESampledSequenceDisplayUnit::Seconds;
+		}
+	}
+}
+
+void SAudioOscilloscopePanelWidget::Construct(const FArguments& InArgs, const FFixedSampledSequenceView& InData, const int32 InNumChannels)
+{
+	NumChannels = InNumChannels;
+
+	DataView = InData;
+
+	PanelLayoutType = InArgs._PanelLayoutType;
+
+	SequenceRulerDisplayUnit    = InArgs._SequenceRulerDisplayUnit;
+	ValueGridOverlayDisplayUnit = InArgs._YAxisLabelsUnit.Get();
+	bHideSequenceRuler          = InArgs._HideSequenceRuler.Get();
+	bHideSequenceGrid           = InArgs._HideSequenceGrid.Get();
+	bHideValueGrid              = InArgs._HideValueGrid.Get();
+	bHideTriggerThresholdLine   = InArgs._HideTriggerThresholdLine.Get();
+
+	check(InArgs._PanelStyle);
+	PanelStyle = InArgs._PanelStyle;
+
+	CreateGridData(PanelStyle->TimeRulerStyle);
+
+	CreateSequenceRuler(SequenceGridData.ToSharedRef(), PanelStyle->TimeRulerStyle);
+	CreateBackground(PanelStyle->WaveViewerStyle);
+	CreateValueGridOverlay(InArgs._ValueGridMaxDivisionParameter, SampledSequenceValueGridOverlay::EGridDivideMode::MidSplit, ValueGridOverlayDisplayUnit, PanelStyle->ValueGridStyle);
+	CreateSequenceViewer(SequenceGridData.ToSharedRef(), DataView, PanelStyle->WaveViewerStyle);
+	CreateTriggerThresholdLine(PanelStyle->TriggerThresholdLineStyle);
+
+	if (PanelLayoutType == EAudioPanelLayoutType::Advanced)
+	{
+		CreateOscilloscopeControls();
+	}
+
+	CreateLayout();
+}
+
+void SAudioOscilloscopePanelWidget::CreateLayout()
+{
+	// Oscilloscope view
+	auto CreateOscilloscopeViewContainer = [this]()
+	{
+		TSharedRef<SOverlay> OscilloscopeViewOverlays = SNew(SOverlay);
+
+		OscilloscopeViewOverlays->AddSlot()
+		[
+			BackgroundBorder.ToSharedRef()
+		];
+
+		if (!bHideValueGrid)
+		{
+			OscilloscopeViewOverlays->AddSlot()
+			[
+				ValueGridOverlay.ToSharedRef()
+			];
+		}
+
+		OscilloscopeViewOverlays->AddSlot()
+		[
+			SequenceViewer.ToSharedRef()
+		];
+
+		OscilloscopeViewOverlays->AddSlot()
+		[
+			TriggerThresholdLineWidget.ToSharedRef()
+		];
+
+		TSharedPtr<SVerticalBox> OscilloscopeViewContainer = SNew(SVerticalBox);
+
+		if (!bHideSequenceRuler)
+		{
+			OscilloscopeViewContainer->AddSlot().AutoHeight()
+			[
+				SequenceRuler.ToSharedRef()
+			];
+		}
+
+		OscilloscopeViewContainer->AddSlot()
+		[
+			OscilloscopeViewOverlays
+		];
+
+		OscilloscopeViewContainer->SetClipping(EWidgetClipping::ClipToBounds);
+
+		return OscilloscopeViewContainer;
+	};
+
+	if (PanelLayoutType == EAudioPanelLayoutType::Advanced)
+	{
+		// Oscilloscope controls
+		auto CreateOscilloscopeControlsContainer = [this]()
+		{
+			TSharedPtr<SVerticalBox> OscilloscopeControlsContainer = SNew(SVerticalBox);
+
+			// Channel Combobox
+			OscilloscopeControlsContainer->AddSlot()
+			.FillHeight(0.2f)
+			.Padding(0.0f, 5.0f, 0.0f, 0.0f)
+			[
+				SNew(SVerticalBox)
+				+ SVerticalBox::Slot()
+				.FillHeight(0.3f)
+				.Padding(0.0f, 0.0f, 0.0f, 4.0f)
+				[
+					SNew(STextBlock)
+					.Text(LOCTEXT("Oscilloscope_Channel_Display_Label", "Channel"))
+					.Justification(ETextJustify::Center)
+				]
+				+ SVerticalBox::Slot()
+				.FillHeight(0.4f)
+				.Padding(5.0f, 0.0f, 5.0f, 0.0f)
+				[
+					ChannelCombobox.ToSharedRef()
+				]
+				+ SVerticalBox::Slot()
+				.FillHeight(0.3f)
+				[
+					SNew(SVerticalBox)
+				]
+			];
+
+			// Trigger Mode Combobox
+			OscilloscopeControlsContainer->AddSlot()
+			.FillHeight(0.2f)
+			[
+				SNew(SVerticalBox)
+				+ SVerticalBox::Slot()
+				.FillHeight(0.3f)
+				.Padding(0.0f, 0.0f, 0.0f, 4.0f)
+				[
+					SNew(STextBlock)
+					.Text(LOCTEXT("Oscilloscope_TriggerMode_Display_Label", "Trigger Mode"))
+					.Justification(ETextJustify::Center)
+				]
+				+ SVerticalBox::Slot()
+				.FillHeight(0.4f)
+				.Padding(5.0f, 0.0f, 5.0f, 0.0f)
+				[
+					TriggerModeCombobox.ToSharedRef()
+				]
+				+ SVerticalBox::Slot()
+				.FillHeight(0.3f)
+				[
+					SNew(SVerticalBox)
+				]
+			];
+
+			// Trigger Threshold knob
+			OscilloscopeControlsContainer->AddSlot()
+			.FillHeight(0.2f)
+			.Padding(0.0f, 0.0f, 0.0f, 8.0f)
+			[
+				SNew(SVerticalBox)
+				+ SVerticalBox::Slot()
+				.FillHeight(0.2f)
+				[
+					SNew(STextBlock)
+					.Text(LOCTEXT("Oscilloscope_TriggerThreshold_Display_Label", "Trigger Threshold"))
+					.Justification(ETextJustify::Center)
+				]
+				+ SVerticalBox::Slot()
+				.FillHeight(0.8f)
+				[
+					TriggerThresholdKnob.ToSharedRef()
+				]
+			];
+
+			// Time Window Knob
+			OscilloscopeControlsContainer->AddSlot()
+			.FillHeight(0.2f)
+			.Padding(0.0f, 0.0f, 0.0f, 5.0f)
+			[
+				SNew(SVerticalBox)
+				+ SVerticalBox::Slot()
+				.FillHeight(0.2f)
+				[
+					SNew(STextBlock)
+					.Text(LOCTEXT("Oscilloscope_TimeWindow_Display_Label", "Time Window"))
+					.Justification(ETextJustify::Center)
+				]
+				+ SVerticalBox::Slot()
+				.FillHeight(0.8f)
+				[
+					TimeWindowKnob.ToSharedRef()
+				]
+			];
+
+			// Analysis Period Knob
+			OscilloscopeControlsContainer->AddSlot()
+			.FillHeight(0.2f)
+			.Padding(0.0f, 0.0f, 0.0f, 5.0f)
+			[
+				SNew(SVerticalBox)
+				+ SVerticalBox::Slot()
+				.FillHeight(0.2f)
+				[
+					SNew(STextBlock)
+					.Text(LOCTEXT("Oscilloscope_AnalysisPeriod_Display_Label", "Analysis Period"))
+					.Justification(ETextJustify::Center)
+				]
+				+ SVerticalBox::Slot()
+				.FillHeight(0.8f)
+				[
+					AnalysisPeriodKnob.ToSharedRef()
+				]
+			];
+
+			return OscilloscopeControlsContainer;
+		};
+
+		OscilloscopeViewProportion = 0.9f;
+
+		ChildSlot
+		[
+			SNew(SHorizontalBox)
+			+ SHorizontalBox::Slot()
+			.FillWidth(OscilloscopeViewProportion)
+			[
+				CreateOscilloscopeViewContainer().ToSharedRef()
+			]
+			+ SHorizontalBox::Slot()
+			.FillWidth(1.0f - OscilloscopeViewProportion)
+			[
+				CreateOscilloscopeControlsContainer().ToSharedRef()
+			]
+		];
+	}
+	else
+	{
+		ChildSlot
+		[
+			CreateOscilloscopeViewContainer().ToSharedRef()
+		];
+	}
+}
+
+void SAudioOscilloscopePanelWidget::CreateGridData(const FFixedSampleSequenceRulerStyle& RulerStyle)
+{
+	SequenceGridData = MakeShared<FFixedSampledSequenceGridData>(DataView.SampleData.Num() / DataView.NumDimensions, DataView.SampleRate, RulerStyle.DesiredWidth, &RulerStyle.TicksTextFont);
+}
+
+void SAudioOscilloscopePanelWidget::CreateSequenceRuler(TSharedRef<FFixedSampledSequenceGridData> InGridData, const FFixedSampleSequenceRulerStyle& RulerStyle)
+{
+	SequenceRuler = SNew(SFixedSampledSequenceRuler, InGridData)
+	.DisplayUnit(AudioOscilloscopePanelWidgetPrivate::ConvertEXAxisLabelsUnitToESampledSequenceDisplayUnit(SequenceRulerDisplayUnit))
+	.Style(&RulerStyle)
+	.DisplayPlayhead(false)
+	.Clipping(EWidgetClipping::ClipToBounds);
+
+	SequenceGridData->OnGridMetricsUpdated.AddSP(SequenceRuler.ToSharedRef(), &SFixedSampledSequenceRuler::UpdateGridMetrics);
+}
+
+void SAudioOscilloscopePanelWidget::CreateBackground(const FSampledSequenceViewerStyle& ViewerStyle)
+{
+	BackgroundBorder = SNew(SBorder)
+	.BorderImage(&ViewerStyle.BackgroundBrush)
+	.BorderBackgroundColor(ViewerStyle.SequenceBackgroundColor);
+}
+
+void SAudioOscilloscopePanelWidget::CreateValueGridOverlay(const uint32 MaxDivisionParameter,
+	const SampledSequenceValueGridOverlay::EGridDivideMode DivideMode,
+	const EYAxisLabelsUnit InValueGridOverlayDisplayUnit,
+	const FSampledSequenceValueGridOverlayStyle& ValueGridStyle)
+{
+	ValueGridOverlay = SNew(SSampledSequenceValueGridOverlay)
+	.MaxDivisionParameter(MaxDivisionParameter)
+	.DivideMode(DivideMode)
+	.ValueGridLabelGenerator(AudioOscilloscopePanelWidgetPrivate::GetValueGridOverlayLabelGenerator(InValueGridOverlayDisplayUnit))
+	.HideLabels(false)
+	.Style(&ValueGridStyle)
+	.SequenceDrawingParams(DrawingParams)
+	.NumDimensions(DataView.NumDimensions);
+}
+
+void SAudioOscilloscopePanelWidget::CreateSequenceViewer(TSharedRef<FFixedSampledSequenceGridData> InGridData, const FFixedSampledSequenceView& InData, const FSampledSequenceViewerStyle& ViewerStyle)
+{
+	SequenceViewer = SNew(SFixedSampledSequenceViewer, InData.SampleData, InData.NumDimensions, InGridData)
+	.Style(&ViewerStyle)
+	.SequenceDrawingParams(DrawingParams)
+	.HideBackground(true)
+	.HideGrid(bHideSequenceGrid);
+
+	SequenceGridData->OnGridMetricsUpdated.AddSP(SequenceViewer.ToSharedRef(), &SFixedSampledSequenceViewer::UpdateGridMetrics);
+}
+
+void SAudioOscilloscopePanelWidget::CreateTriggerThresholdLine(const FTriggerThresholdLineStyle& TriggerThresholdLineStyle)
+{
+	TriggerThresholdLineWidget = SNew(STriggerThresholdLineWidget)
+	.Style(&TriggerThresholdLineStyle)
+	.Visibility(bHideTriggerThresholdLine ? EVisibility::Hidden : EVisibility::Visible)
+	.SequenceDrawingParams(DrawingParams);
+}
+
+void SAudioOscilloscopePanelWidget::CreateChannelCombobox()
+{
+	auto OnGenerateChannelCombobox = [this](TSharedPtr<FString> InChannel)
+	{
+		return SNew(STextBlock).Text(FText::FromString(*InChannel));
+	};
+
+	auto OnChannelSelectionChanged = [this](TSharedPtr<FString> InChannel, ESelectInfo::Type InSelectInfo)
+	{
+		SelectedChannelPtr = InChannel;
+		ChannelCombobox->SetSelectedItem(SelectedChannelPtr);
+
+		OnSelectedChannelChanged.Broadcast(FCString::Atoi(**InChannel.Get()));
+	};
+
+	auto GetSelectedChannel = [this]()
+	{
+		return SelectedChannelPtr.IsValid() ? FText::FromString(*SelectedChannelPtr) : FText::GetEmpty();
+	};
+
+	for (uint32 Index = 0; Index < NumChannels; ++Index)
+	{
+		ChannelComboboxOptionsSource.Add(MakeShareable(new FString(FString::FromInt(Index + 1))));
+	}
+
+	SelectedChannelPtr = ChannelComboboxOptionsSource[0];
+
+	ChannelCombobox = SNew(SComboBox<TSharedPtr<FString>>)
+	.OptionsSource(&ChannelComboboxOptionsSource)
+	.OnGenerateWidget_Lambda(OnGenerateChannelCombobox)
+	.OnSelectionChanged_Lambda(OnChannelSelectionChanged)
+	.InitiallySelectedItem(SelectedChannelPtr)
+	[
+		SNew(STextBlock).Text_Lambda(GetSelectedChannel)
+	];
+}
+
+void SAudioOscilloscopePanelWidget::CreateTriggerModeCombobox()
+{
+	auto OnGenerateTriggerModeCombobox = [this](TSharedPtr<EAudioOscilloscopeTriggerMode> InTriggerMode)
+	{
+		const FString InTriggerModeEnumStr = StaticEnum<EAudioOscilloscopeTriggerMode>()->GetNameStringByValue(static_cast<uint8>(*InTriggerMode));
+
+		return SNew(STextBlock).Text(FText::FromString(*InTriggerModeEnumStr));
+	};
+
+	auto OnTriggerModeSelectionChanged = [this](TSharedPtr<EAudioOscilloscopeTriggerMode> InTriggerMode, ESelectInfo::Type InSelectInfo)
+	{
+		SelectedTriggerModePtr = InTriggerMode;
+		TriggerModeCombobox->SetSelectedItem(SelectedTriggerModePtr);
+
+		const EVisibility VisibilityMode = (*SelectedTriggerModePtr.Get() != EAudioOscilloscopeTriggerMode::None) ? EVisibility::Visible : EVisibility::Hidden;
+		TriggerThresholdLineWidget->SetVisibility(VisibilityMode);
+
+		OnTriggerModeChanged.Broadcast(*InTriggerMode.Get());
+	};
+
+	auto GetSelectedTriggerMode = [this]()
+	{
+		const FString SelectedTriggerModeEnumStr = StaticEnum<EAudioOscilloscopeTriggerMode>()->GetNameStringByValue(static_cast<uint8>(*SelectedTriggerModePtr));
+
+		return SelectedTriggerModePtr.IsValid() ? FText::FromString(*SelectedTriggerModeEnumStr) : FText::GetEmpty();
+	};
+
+	for (int64 Index = 0; Index < StaticEnum<EAudioOscilloscopeTriggerMode>()->GetMaxEnumValue(); ++Index)
+	{
+		TriggerModeComboboxOptionsSource.Add(MakeShareable(new EAudioOscilloscopeTriggerMode(static_cast<EAudioOscilloscopeTriggerMode>(Index))));
+	}
+
+	SelectedTriggerModePtr = TriggerModeComboboxOptionsSource[0];
+
+	TriggerModeCombobox = SNew(SComboBox<TSharedPtr<EAudioOscilloscopeTriggerMode>>)
+	.OptionsSource(&TriggerModeComboboxOptionsSource)
+	.OnGenerateWidget_Lambda(OnGenerateTriggerModeCombobox)
+	.OnSelectionChanged_Lambda(OnTriggerModeSelectionChanged)
+	.InitiallySelectedItem(SelectedTriggerModePtr)
+	[
+		SNew(STextBlock).Text_Lambda(GetSelectedTriggerMode)
+	];
+}
+
+void SAudioOscilloscopePanelWidget::CreateTriggerThresholdKnob()
+{
+	auto OnValueChangedLambda = [this](float Value)
+	{
+		if (!bIsInputWidgetTransacting)
+		{
+		#if WITH_EDITOR
+			if (GEditor)
+			{
+				GEditor->BeginTransaction(LOCTEXT("Oscilloscope_TriggerThreshold_Knob_Changed_Msg", "Set oscilloscope Trigger Threshold value."));
+			}
+		#endif
+			bIsInputWidgetTransacting = true;
+		}
+
+		if (const float TriggerThresholdKnobValue = TriggerThresholdKnob->GetOutputValue(Value); 
+			TriggerThresholdKnobValue != TriggerThresholdValue)
+		{
+			TriggerThresholdValue = TriggerThresholdKnobValue;
+
+			TriggerThresholdLineWidget->SetTriggerThreshold(TriggerThresholdValue);
+		}
+	};
+
+	auto OnRadialSliderMouseCaptureBeginLambda = [this]()
+	{
+		if (!bIsInputWidgetTransacting)
+		{
+		#if WITH_EDITOR
+			if (GEditor)
+			{
+				GEditor->BeginTransaction(LOCTEXT("Oscilloscope_TriggerThreshold_Knob_CaptureBegin_Msg", "Set oscilloscope Trigger Threshold value."));
+			}
+		#endif
+			bIsInputWidgetTransacting = true;
+		}
+	};
+
+	auto OnRadialSliderMouseCaptureEndLambda = [this]()
+	{
+		if (bIsInputWidgetTransacting)
+		{
+		#if WITH_EDITOR
+			if (GEditor)
+			{
+				GEditor->EndTransaction();
+			}
+		#endif
+			bIsInputWidgetTransacting = false;
+
+			OnTriggerThresholdChanged.Broadcast(TriggerThresholdValue);
+		}
+		else
+		{
+			UE_LOG(LogAudioWidgets, Warning, TEXT("Unmatched oscilloscope widget transaction."));
+		}
+	};
+
+	auto OnRadialSliderEnabledLambda = [this]()
+	{
+		return *SelectedTriggerModePtr.Get() != EAudioOscilloscopeTriggerMode::None;
+	};
+
+	TriggerThresholdKnob = SNew(SAudioRadialSlider)
+	.OnMouseCaptureBegin_Lambda(OnRadialSliderMouseCaptureBeginLambda)
+	.OnMouseCaptureEnd_Lambda(OnRadialSliderMouseCaptureEndLambda)
+	.OnValueChanged_Lambda(OnValueChangedLambda)
+	.SliderValue(0.5f)
+	.IsEnabled_Lambda(OnRadialSliderEnabledLambda); // In normalized range, range will change below
+
+	TriggerThresholdKnob->SetShowUnitsText(false);
+	TriggerThresholdKnob->SetOutputRange(FVector2D(-1.0f, 1.0f));
+}
+
+void SAudioOscilloscopePanelWidget::CreateTimeWindowKnob()
+{
+	auto OnValueChangedLambda = [this](float Value)
+	{
+		if (!bIsInputWidgetTransacting)
+		{
+		#if WITH_EDITOR
+			if (GEditor)
+			{
+				GEditor->BeginTransaction(LOCTEXT("Oscilloscope_TimeWindow_Knob_Changed_Msg", "Set oscilloscope Time Window value."));
+			}
+		#endif
+			bIsInputWidgetTransacting = true;
+		}
+
+		if (const float TimeWindowKnobValue = TimeWindowKnob->GetOutputValue(Value);
+			TimeWindowKnobValue != TimeWindowValue)
+		{
+			TimeWindowValue = TimeWindowKnobValue;
+
+			OnTimeWindowValueChanged.Broadcast(TimeWindowValue);
+		}
+	};
+
+	auto OnRadialSliderMouseCaptureBeginLambda = [this]()
+	{
+		if (!bIsInputWidgetTransacting)
+		{
+		#if WITH_EDITOR
+			if (GEditor)
+			{
+				GEditor->BeginTransaction(LOCTEXT("Oscilloscope_TimeWindow_Knob_CaptureBegin_Msg", "Set oscilloscope Time Window value."));
+			}
+		#endif
+			bIsInputWidgetTransacting = true;
+		}
+	};
+
+	auto OnRadialSliderMouseCaptureEndLambda = [this]()
+	{
+		if (bIsInputWidgetTransacting)
+		{
+		#if WITH_EDITOR
+			if (GEditor)
+			{
+				GEditor->EndTransaction();
+			}
+		#endif
+			bIsInputWidgetTransacting = false;
+		}
+		else
+		{
+			UE_LOG(LogAudioWidgets, Warning, TEXT("Unmatched oscilloscope widget transaction."));
+		}
+	};
+
+	TimeWindowKnob = SNew(SAudioRadialSlider)
+	.OnMouseCaptureBegin_Lambda(OnRadialSliderMouseCaptureBeginLambda)
+	.OnMouseCaptureEnd_Lambda(OnRadialSliderMouseCaptureEndLambda);
+
+	TimeWindowKnob->SetOutputRange(FVector2D(10.0f, 5000.0f));
+	TimeWindowKnob->SetUnitsText(FText::FromString("ms"));
+
+	TimeWindowKnob->OnValueChanged.BindLambda(OnValueChangedLambda);
+}
+
+void SAudioOscilloscopePanelWidget::CreateAnalysisPeriodKnob()
+{
+	auto OnValueChangedLambda = [this](float Value)
+	{
+		if (!bIsInputWidgetTransacting)
+		{
+		#if WITH_EDITOR
+			if (GEditor)
+			{
+				GEditor->BeginTransaction(LOCTEXT("Oscilloscope_AnalysisPeriod_Knob_Changed_Msg", "Set oscilloscope Analysis Period value."));
+			}
+		#endif
+			bIsInputWidgetTransacting = true;
+		}
+
+		if (const float AnalysisPeriodKnobValue = AnalysisPeriodKnob->GetOutputValue(Value);
+			AnalysisPeriodKnobValue != AnalysisPeriodValue)
+		{
+			AnalysisPeriodValue = AnalysisPeriodKnobValue;
+
+			OnAnalysisPeriodChanged.Broadcast(AnalysisPeriodValue);
+		}
+	};
+
+	auto OnRadialSliderMouseCaptureBeginLambda = [this]()
+	{
+		if (!bIsInputWidgetTransacting)
+		{
+		#if WITH_EDITOR
+			if (GEditor)
+			{
+				GEditor->BeginTransaction(LOCTEXT("Oscilloscope_AnalysisPeriod_Knob_CaptureBegin_Msg", "Set oscilloscope Analysis Period value."));
+			}
+		#endif
+			bIsInputWidgetTransacting = true;
+		}
+	};
+
+	auto OnRadialSliderMouseCaptureEndLambda = [this]()
+	{
+		if (bIsInputWidgetTransacting)
+		{
+		#if WITH_EDITOR
+			if (GEditor)
+			{
+				GEditor->EndTransaction();
+			}
+		#endif
+			bIsInputWidgetTransacting = false;
+		}
+		else
+		{
+			UE_LOG(LogAudioWidgets, Warning, TEXT("Unmatched oscilloscope widget transaction."));
+		}
+	};
+
+	AnalysisPeriodKnob = SNew(SAudioRadialSlider)
+	.OnMouseCaptureBegin_Lambda(OnRadialSliderMouseCaptureBeginLambda)
+	.OnMouseCaptureEnd_Lambda(OnRadialSliderMouseCaptureEndLambda);
+
+	AnalysisPeriodKnob->SetOutputRange(FVector2D(10.0f, 1000.0f));
+	AnalysisPeriodKnob->SetUnitsText(FText::FromString("ms"));
+
+	AnalysisPeriodKnob->OnValueChanged.BindLambda(OnValueChangedLambda);
+}
+
+void SAudioOscilloscopePanelWidget::CreateOscilloscopeControls()
+{
+	CreateChannelCombobox();
+	CreateTriggerModeCombobox();
+	CreateTriggerThresholdKnob();
+	CreateTimeWindowKnob();
+	CreateAnalysisPeriodKnob();
+}
+
+void SAudioOscilloscopePanelWidget::SetSequenceRulerDisplayUnit(const EXAxisLabelsUnit InDisplayUnit)
+{
+	SequenceRulerDisplayUnit = InDisplayUnit;
+	SequenceRuler->UpdateDisplayUnit(AudioOscilloscopePanelWidgetPrivate::ConvertEXAxisLabelsUnitToESampledSequenceDisplayUnit(SequenceRulerDisplayUnit));
+}
+
+void SAudioOscilloscopePanelWidget::SetXAxisGridVisibility(const bool InbIsVisible)
+{
+	SequenceViewer->SetHideGrid(!InbIsVisible);
+}
+
+void SAudioOscilloscopePanelWidget::SetValueGridOverlayDisplayUnit(const EYAxisLabelsUnit InDisplayUnit)
+{
+	ValueGridOverlayDisplayUnit = InDisplayUnit;
+	ValueGridOverlay->SetLabelGenerator(AudioOscilloscopePanelWidgetPrivate::GetValueGridOverlayLabelGenerator(InDisplayUnit));
+}
+
+void SAudioOscilloscopePanelWidget::SetYAxisGridVisibility(const bool InbIsVisible)
+{
+	ValueGridOverlay->SetHideGrid(!InbIsVisible);
+}
+
+void SAudioOscilloscopePanelWidget::SetYAxisLabelsVisibility(const bool InbIsVisible)
+{
+	ValueGridOverlay->SetHideLabels(!InbIsVisible);
+}
+
+void SAudioOscilloscopePanelWidget::SetTriggerThreshold(float InTriggerThreshold)
+{
+	TriggerThresholdLineWidget->SetTriggerThreshold(InTriggerThreshold);
+}
+
+void SAudioOscilloscopePanelWidget::SetTriggerThresholdVisibility(const bool InbIsVisible)
+{
+	TriggerThresholdLineWidget->SetVisibility(InbIsVisible ? EVisibility::Visible : EVisibility::Hidden);
+}
+
+void SAudioOscilloscopePanelWidget::UpdateSequenceRulerStyle(const FFixedSampleSequenceRulerStyle UpdatedStyle)
+{
+	SequenceRuler->OnStyleUpdated(UpdatedStyle);
+}
+
+void SAudioOscilloscopePanelWidget::UpdateValueGridOverlayStyle(const FSampledSequenceValueGridOverlayStyle UpdatedStyle)
+{
+	ValueGridOverlay->OnStyleUpdated(UpdatedStyle);
+}
+
+void SAudioOscilloscopePanelWidget::UpdateSequenceViewerStyle(const FSampledSequenceViewerStyle UpdatedStyle)
+{
+	SequenceViewer->OnStyleUpdated(UpdatedStyle);
+	BackgroundBorder->SetBorderBackgroundColor(UpdatedStyle.SequenceBackgroundColor);
+}
+
+void SAudioOscilloscopePanelWidget::UpdateTriggerThresholdStyle(const FTriggerThresholdLineStyle UpdatedStyle)
+{
+	TriggerThresholdLineWidget->OnStyleUpdated(UpdatedStyle);
+}
+
+void SAudioOscilloscopePanelWidget::Tick(const FGeometry& AllottedGeometry, const double InCurrentTime, const float InDeltaTime)
+{
+	const float PaintedWidth = AllottedGeometry.GetLocalSize().X;
+
+	if (PaintedWidth != CachedPixelWidth)
+	{
+		CachedPixelWidth = PaintedWidth;
+		SequenceGridData->UpdateGridMetrics(PaintedWidth * OscilloscopeViewProportion);
+	}
+}
+
+void SAudioOscilloscopePanelWidget::ReceiveSequenceView(const FFixedSampledSequenceView InData, const uint32 FirstSampleIndex)
+{
+	const uint32 FirstRenderedFrame = FirstSampleIndex / InData.NumDimensions;
+	const uint32 NumFrames          = InData.SampleData.Num() / InData.NumDimensions;
+
+	SequenceGridData->UpdateDisplayRange(TRange<uint32>(FirstRenderedFrame, FirstRenderedFrame + NumFrames));
+
+	SequenceViewer->UpdateView(InData.SampleData, InData.NumDimensions);
+}
+
+#undef LOCTEXT_NAMESPACE
diff --git a/Engine/Plugins/Runtime/AudioWidgets/Source/AudioWidgets/Private/SAudioVectorscopePanelWidget.cpp b/Engine/Plugins/Runtime/AudioWidgets/Source/AudioWidgets/Private/SAudioVectorscopePanelWidget.cpp
new file mode 100644
--- /dev/null
+++ b/Engine/Plugins/Runtime/AudioWidgets/Source/AudioWidgets/Private/SAudioVectorscopePanelWidget.cpp
@@ -0,0 +1,394 @@
+// Copyright Epic Games, Inc. All Rights Reserved.
+
+#include "SAudioVectorscopePanelWidget.h"
+
+#include "AudioWidgetsLog.h"
+#include "DSP/Dsp.h"
+#include "SAudioRadialSlider.h"
+#include "SFixedSampledSequenceVectorViewer.h"
+#include "Widgets/Colors/SColorBlock.h"
+#include "Widgets/Layout/SBorder.h"
+#include "Widgets/SBoxPanel.h"
+#include "Widgets/SOverlay.h"
+#include "Widgets/Text/STextBlock.h"
+
+#if WITH_EDITOR
+#include "Editor.h"
+#endif
+
+#define LOCTEXT_NAMESPACE "SAudioVectorscopePanelWidget"
+
+void SAudioVectorscopePanelWidget::Construct(const FArguments& InArgs, const FFixedSampledSequenceView& InData)
+{
+	DataView = InData;
+
+	PanelLayoutType = InArgs._PanelLayoutType;
+
+	bHideValueGrid = InArgs._HideGrid.Get();
+
+	check(InArgs._PanelStyle);
+	PanelStyle = InArgs._PanelStyle;
+
+	CreateBackground(PanelStyle->VectorViewerStyle);
+
+	ValueGridOverlayXAxis = CreateValueGridOverlay(InArgs._ValueGridMaxDivisionParameter,
+		SampledSequenceValueGridOverlay::EGridDivideMode::MidSplit,
+		PanelStyle->ValueGridStyle,
+		SampledSequenceDrawingUtils::ESampledSequenceDrawOrientation::Horizontal);
+
+	ValueGridOverlayYAxis = CreateValueGridOverlay(InArgs._ValueGridMaxDivisionParameter,
+		SampledSequenceValueGridOverlay::EGridDivideMode::MidSplit,
+		PanelStyle->ValueGridStyle,
+		SampledSequenceDrawingUtils::ESampledSequenceDrawOrientation::Vertical);
+	
+	CreateSequenceVectorViewer(DataView, PanelStyle->VectorViewerStyle);
+
+	if (PanelLayoutType == EAudioPanelLayoutType::Advanced)
+	{
+		CreateVectorscopeControls();
+	}
+
+	CreateLayout();
+}
+
+void SAudioVectorscopePanelWidget::CreateLayout()
+{
+	// Vectorscope view
+	auto CreateVectorscopeViewContainer = [this]()
+	{
+		TSharedPtr<SOverlay> VectorscopeViewOverlays = SNew(SOverlay);
+
+		VectorscopeViewOverlays->AddSlot()
+		[
+			BackgroundBorder.ToSharedRef()
+		];
+
+		if (!bHideValueGrid)
+		{
+			VectorscopeViewOverlays->AddSlot()
+			[
+				ValueGridOverlayXAxis.ToSharedRef()
+			];
+
+			VectorscopeViewOverlays->AddSlot()
+			[
+				ValueGridOverlayYAxis.ToSharedRef()
+			];
+		}
+
+		VectorscopeViewOverlays->AddSlot()
+		[
+			SequenceVectorViewer.ToSharedRef()
+		];
+
+		TSharedPtr<SVerticalBox> VectorscopeViewContainer = SNew(SVerticalBox);
+
+		VectorscopeViewContainer->AddSlot()
+		[
+			VectorscopeViewOverlays.ToSharedRef()
+		];
+
+		VectorscopeViewContainer->SetClipping(EWidgetClipping::ClipToBounds);
+
+		return VectorscopeViewContainer;
+	};
+
+	if (PanelLayoutType == EAudioPanelLayoutType::Advanced)
+	{
+		// Vectorscope controls
+		auto CreateVectorscopeControlsContainer = [this]()
+		{
+			TSharedPtr<SVerticalBox> VectorscopeControlsContainer = SNew(SVerticalBox);
+
+			// Time Window Knob
+			VectorscopeControlsContainer->AddSlot()
+			.FillHeight(0.2f)
+			.Padding(0.0f, 0.0f, 0.0f, 5.0f)
+			[
+				SNew(SVerticalBox)
+				+ SVerticalBox::Slot()
+				.FillHeight(0.2f)
+				[
+					SNew(STextBlock)
+					.Text(LOCTEXT("Vectorscope_TimeWindow_Display_Label", "Persistence"))
+					.Justification(ETextJustify::Center)
+				]
+				+ SVerticalBox::Slot()
+				.FillHeight(0.8f)
+				[
+					TimeWindowKnob.ToSharedRef()
+				]
+			];
+
+			// Scale Knob
+			VectorscopeControlsContainer->AddSlot()
+			.FillHeight(0.2f)
+			.Padding(0.0f, 0.0f, 0.0f, 5.0f)
+			[
+				SNew(SVerticalBox)
+				+ SVerticalBox::Slot()
+				.FillHeight(0.2f)
+				[
+					SNew(STextBlock)
+					.Text(LOCTEXT("Vectorscope_Scale_Display_Label", "Scale"))
+					.Justification(ETextJustify::Center)
+				]
+				+ SVerticalBox::Slot()
+				.FillHeight(0.8f)
+				[
+					ScaleKnob.ToSharedRef()
+				]
+			];
+
+			VectorscopeControlsContainer->AddSlot()
+			.FillHeight(0.6f)
+			[
+				SNew(SVerticalBox)
+			];
+
+			return VectorscopeControlsContainer;
+		};
+
+		VectorscopeViewProportion = 0.9f;
+
+		ChildSlot
+		[
+			SNew(SHorizontalBox)
+			+ SHorizontalBox::Slot()
+			.FillWidth(VectorscopeViewProportion)
+			[
+				CreateVectorscopeViewContainer().ToSharedRef()
+			]
+			+ SHorizontalBox::Slot()
+			.FillWidth(1.0f - VectorscopeViewProportion)
+			[
+				CreateVectorscopeControlsContainer().ToSharedRef()
+			]
+		];
+	}
+	else
+	{
+		ChildSlot
+		[
+			CreateVectorscopeViewContainer().ToSharedRef()
+		];
+	}
+}
+
+void SAudioVectorscopePanelWidget::CreateBackground(const FSampledSequenceVectorViewerStyle& VectorViewerStyle)
+{
+	BackgroundBorder = SNew(SBorder)
+	.BorderImage(&VectorViewerStyle.BackgroundBrush)
+	.BorderBackgroundColor(VectorViewerStyle.BackgroundColor);
+}
+
+TSharedPtr<SSampledSequenceValueGridOverlay> SAudioVectorscopePanelWidget::CreateValueGridOverlay(const uint32 MaxDivisionParameter,
+	const SampledSequenceValueGridOverlay::EGridDivideMode DivideMode,
+	const FSampledSequenceValueGridOverlayStyle& ValueGridStyle,
+	const SampledSequenceDrawingUtils::ESampledSequenceDrawOrientation GridOrientation)
+{
+	using namespace SampledSequenceDrawingUtils;
+
+	FSampledSequenceDrawingParams ValueGridOverlayDrawingParams;
+	ValueGridOverlayDrawingParams.Orientation = GridOrientation;
+
+	return SNew(SSampledSequenceValueGridOverlay)
+		   .MaxDivisionParameter(MaxDivisionParameter)
+		   .DivideMode(DivideMode)
+		   .HideLabels(true)
+		   .Style(&ValueGridStyle)
+		   .SequenceDrawingParams(MoveTemp(ValueGridOverlayDrawingParams))
+		   .NumDimensions(1);
+}
+
+void SAudioVectorscopePanelWidget::CreateSequenceVectorViewer(const FFixedSampledSequenceView& InData, const FSampledSequenceVectorViewerStyle& VectorViewerStyle)
+{
+	SequenceVectorViewer = SNew(SFixedSampledSequenceVectorViewer, InData.SampleData, InData.NumDimensions)
+	.Style(&VectorViewerStyle)
+	.SequenceDrawingParams(SampledSequenceDrawingUtils::FSampledSequenceDrawingParams());
+}
+
+void SAudioVectorscopePanelWidget::CreateTimeWindowKnob()
+{
+	auto OnValueChangedLambda = [this](float Value)
+	{
+		if (TimeWindowKnob.IsValid())
+		{
+			if (!bIsInputWidgetTransacting)
+			{
+			#if WITH_EDITOR
+				if (GEditor)
+				{
+					GEditor->BeginTransaction(LOCTEXT("Vectorscope_TimeWindow_Knob_Changed_Msg", "Set vectorscope Time Window value."));
+				}
+			#endif
+				bIsInputWidgetTransacting = true;
+			}
+
+			if (const float TimeWindowKnobValue = TimeWindowKnob->GetOutputValue(Value);
+				TimeWindowKnobValue != TimeWindowValue)
+			{
+				TimeWindowValue = TimeWindowKnobValue;
+
+				OnTimeWindowValueChanged.Broadcast(TimeWindowValue);
+			}
+		}
+	};
+
+	auto OnRadialSliderMouseCaptureBeginLambda = [this]()
+	{
+		if (!bIsInputWidgetTransacting)
+		{
+		#if WITH_EDITOR
+			if (GEditor)
+			{
+				GEditor->BeginTransaction(LOCTEXT("Vectorscope_TimeWindow_Knob_CaptureBegin_Msg", "Set vectorscope Time Window value."));
+			}
+		#endif
+			bIsInputWidgetTransacting = true;
+		}
+	};
+
+	auto OnRadialSliderMouseCaptureEndLambda = [this]()
+	{
+		if (bIsInputWidgetTransacting)
+		{
+		#if WITH_EDITOR
+			if (GEditor)
+			{
+				GEditor->EndTransaction();
+			}
+		#endif
+			bIsInputWidgetTransacting = false;
+		}
+		else
+		{
+			UE_LOG(LogAudioWidgets, Warning, TEXT("Unmatched vectorscope widget transaction."));
+		}
+	};
+
+	TimeWindowKnob = SNew(SAudioRadialSlider)
+	.OnMouseCaptureBegin_Lambda(OnRadialSliderMouseCaptureBeginLambda)
+	.OnMouseCaptureEnd_Lambda(OnRadialSliderMouseCaptureEndLambda)
+	.SliderValue(0.0f);
+
+	TimeWindowKnob->SetOutputRange(FVector2D(10.0f, 500.0f));
+	TimeWindowKnob->SetUnitsText(FText::FromString("ms"));
+
+	TimeWindowKnob->OnValueChanged.BindLambda(OnValueChangedLambda);
+}
+
+void SAudioVectorscopePanelWidget::CreateScaleKnob()
+{
+	auto OnValueChangedLambda = [this](float Value)
+	{
+		if (ScaleKnob.IsValid())
+		{
+			if (!bIsInputWidgetTransacting)
+			{
+			#if WITH_EDITOR
+				if (GEditor)
+				{
+					GEditor->BeginTransaction(LOCTEXT("Vectorscope_Scale_Knob_Changed_Msg", "Set vectorscope Scale value."));
+				}
+			#endif
+				bIsInputWidgetTransacting = true;
+			}
+
+			if (const float ScaleKnobValue = ScaleKnob->GetOutputValue(Value);
+				ScaleKnobValue != ScaleValue)
+			{
+				ScaleValue = ScaleKnobValue;
+
+				SequenceVectorViewer->SetScaleFactor(ScaleValue);
+			}
+		}
+	};
+
+	auto OnRadialSliderMouseCaptureBeginLambda = [this]()
+	{
+		if (!bIsInputWidgetTransacting)
+		{
+		#if WITH_EDITOR
+			if (GEditor)
+			{
+				GEditor->BeginTransaction(LOCTEXT("Vectorscope_Scale_Knob_CaptureBegin_Msg", "Set vectorscope Scale value."));
+			}
+		#endif
+			bIsInputWidgetTransacting = true;
+		}
+	};
+
+	auto OnRadialSliderMouseCaptureEndLambda = [this]()
+	{
+		if (bIsInputWidgetTransacting)
+		{
+		#if WITH_EDITOR
+			if (GEditor)
+			{
+				GEditor->EndTransaction();
+			}
+		#endif
+			bIsInputWidgetTransacting = false;
+		}
+		else
+		{
+			UE_LOG(LogAudioWidgets, Warning, TEXT("Unmatched vectorscope widget transaction."));
+		}
+	};
+
+	ScaleKnob = SNew(SAudioRadialSlider)
+	.OnMouseCaptureBegin_Lambda(OnRadialSliderMouseCaptureBeginLambda)
+	.OnMouseCaptureEnd_Lambda(OnRadialSliderMouseCaptureEndLambda)
+	.SliderValue(1.0f);
+
+	ScaleKnob->SetOutputRange(FVector2D(0.0f, 1.0f));
+	ScaleKnob->SetShowUnitsText(false);
+
+	ScaleKnob->OnValueChanged.BindLambda(OnValueChangedLambda);
+}
+
+void SAudioVectorscopePanelWidget::CreateVectorscopeControls()
+{
+	CreateTimeWindowKnob();
+	CreateScaleKnob();
+}
+
+void SAudioVectorscopePanelWidget::SetGridVisibility(const bool InbIsVisible)
+{
+	ValueGridOverlayXAxis->SetHideGrid(!InbIsVisible);
+	ValueGridOverlayYAxis->SetHideGrid(!InbIsVisible);
+}
+
+void SAudioVectorscopePanelWidget::SetValueGridOverlayMaxNumDivisions(const uint32 InGridMaxNumDivisions)
+{
+	ValueGridOverlayXAxis->SetMaxDivisionParameter(InGridMaxNumDivisions);
+	ValueGridOverlayYAxis->SetMaxDivisionParameter(InGridMaxNumDivisions);
+}
+
+void SAudioVectorscopePanelWidget::SetVectorViewerScaleFactor(const float InScaleFactor)
+{
+	SequenceVectorViewer->SetScaleFactor(InScaleFactor);
+}
+
+void SAudioVectorscopePanelWidget::UpdateValueGridOverlayStyle(const FSampledSequenceValueGridOverlayStyle UpdatedStyle)
+{
+	ValueGridOverlayXAxis->OnStyleUpdated(UpdatedStyle);
+	ValueGridOverlayYAxis->OnStyleUpdated(UpdatedStyle);
+}
+
+void SAudioVectorscopePanelWidget::UpdateSequenceVectorViewerStyle(const FSampledSequenceVectorViewerStyle UpdatedStyle)
+{
+	SequenceVectorViewer->OnStyleUpdated(UpdatedStyle);
+	BackgroundBorder->SetBorderBackgroundColor(UpdatedStyle.BackgroundColor);
+}
+
+void SAudioVectorscopePanelWidget::ReceiveSequenceView(const FFixedSampledSequenceView InData, const uint32 FirstSampleIndex)
+{
+	if (SequenceVectorViewer)
+	{
+		SequenceVectorViewer->UpdateView(InData.SampleData, InData.NumDimensions);
+	}
+}
+
+#undef LOCTEXT_NAMESPACE
\ No newline at end of file
diff --git a/Engine/Plugins/Runtime/AudioWidgets/Source/AudioWidgets/Private/SFixedSampledSequenceVectorViewer.cpp b/Engine/Plugins/Runtime/AudioWidgets/Source/AudioWidgets/Private/SFixedSampledSequenceVectorViewer.cpp
new file mode 100644
--- /dev/null
+++ b/Engine/Plugins/Runtime/AudioWidgets/Source/AudioWidgets/Private/SFixedSampledSequenceVectorViewer.cpp
@@ -0,0 +1,98 @@
+// Copyright Epic Games, Inc. All Rights Reserved.
+
+#include "SFixedSampledSequenceVectorViewer.h"
+
+#include "AudioWidgetsSlateTypes.h"
+
+#define LOCTEXT_NAMESPACE "OscilloscopeWidget"
+
+void SFixedSampledSequenceVectorViewer::Construct(const FArguments& InArgs, TArrayView<const float> InSampleData, const uint8 InNumChannels)
+{
+	UpdateView(InSampleData, InNumChannels);
+	
+	check(InArgs._Style);
+
+	DrawingParams = InArgs._SequenceDrawingParams;
+
+	LineColor     = InArgs._Style->LineColor;
+	LineThickness = InArgs._Style->LineThickness;
+
+	ScaleFactor = InArgs._ScaleFactor;
+}
+
+void SFixedSampledSequenceVectorViewer::UpdateView(TArrayView<const float> InSampleData, const uint8 InNumChannels)
+{
+	SampleData   = InSampleData;
+	NumChannels  = InNumChannels;
+	bForceRedraw = true;
+}
+
+FVector2D SFixedSampledSequenceVectorViewer::ComputeDesiredSize(float) const
+{
+	return FVector2D(DesiredWidth, DesiredHeight);
+}
+
+int32 SFixedSampledSequenceVectorViewer::OnPaint(const FPaintArgs& Args,
+	const FGeometry& AllottedGeometry, 
+	const FSlateRect& MyCullingRect, 
+	FSlateWindowElementList& OutDrawElements, 
+	int32 LayerId, 
+	const FWidgetStyle& InWidgetStyle, 
+	bool bParentEnabled) const
+{
+	check(NumChannels > 0);
+
+	TArray<FVector2D> PointsToDraw;
+
+	PointsToDraw.Reserve(SampleData.Num() / NumChannels);
+
+	const float MaxScaleFactor   = DrawingParams.MaxSequenceHeightRatio * 0.5f;
+	const float FinalScaleFactor = ScaleFactor * MaxScaleFactor;
+
+	const FScale2D Scale(AllottedGeometry.GetLocalSize() * FVector2D(FinalScaleFactor, -FinalScaleFactor));
+	const FVector2D Translation(AllottedGeometry.GetLocalSize().X / 2.0, AllottedGeometry.GetLocalSize().Y / 2.0);
+
+	const FTransform2D PointsTransform(Scale, Translation);
+
+	if (NumChannels == 1)
+	{
+		for (int32 Index = 0; Index < SampleData.Num(); ++Index)
+		{
+			PointsToDraw.Add(PointsTransform.TransformPoint(FVector2D(SampleData[Index], 0.0)));
+		}
+	}
+	else
+	{
+		// For more than 2 channels it will only take the first 2 channels of audio data
+		for (int32 Index = 0; Index < SampleData.Num() - NumChannels + 1; Index += NumChannels)
+		{
+			PointsToDraw.Add(PointsTransform.TransformPoint(FVector2D(SampleData[Index], SampleData[Index + 1])));
+		}
+	}
+
+	FSlateDrawElement::MakeLines(OutDrawElements,
+		LayerId,
+		AllottedGeometry.ToPaintGeometry(),
+		PointsToDraw,
+		ESlateDrawEffect::None,
+		LineColor,
+		true,
+		LineThickness);
+
+	++LayerId;
+
+	return LayerId;
+}
+
+void SFixedSampledSequenceVectorViewer::OnStyleUpdated(const FSampledSequenceVectorViewerStyle UpdatedStyle)
+{
+	LineColor     = UpdatedStyle.LineColor;
+	LineThickness = UpdatedStyle.LineThickness;
+}
+
+void SFixedSampledSequenceVectorViewer::SetScaleFactor(const float InScaleFactor)
+{
+	ScaleFactor = InScaleFactor;
+}
+
+#undef LOCTEXT_NAMESPACE
diff --git a/Engine/Plugins/Runtime/AudioWidgets/Source/AudioWidgets/Private/STriggerThresholdLineWidget.cpp b/Engine/Plugins/Runtime/AudioWidgets/Source/AudioWidgets/Private/STriggerThresholdLineWidget.cpp
new file mode 100644
--- /dev/null
+++ b/Engine/Plugins/Runtime/AudioWidgets/Source/AudioWidgets/Private/STriggerThresholdLineWidget.cpp
@@ -0,0 +1,51 @@
+// Copyright Epic Games, Inc. All Rights Reserved.
+
+#include "STriggerThresholdLineWidget.h"
+
+void STriggerThresholdLineWidget::Construct(const FArguments& InArgs)
+{
+	TriggerThreshold = InArgs._TriggerThreshold;
+	LineColor        = InArgs._Style->LineColor;
+
+	DrawingParams = InArgs._SequenceDrawingParams;
+}
+
+int32 STriggerThresholdLineWidget::OnPaint(const FPaintArgs& Args,
+	const FGeometry& AllottedGeometry, 
+	const FSlateRect& MyCullingRect, 
+	FSlateWindowElementList& OutDrawElements, 
+	int32 LayerId, 
+	const FWidgetStyle& InWidgetStyle, 
+	bool bParentEnabled) const
+{
+	const SampledSequenceDrawingUtils::FDimensionSlot DimensionSlot(0, 1, AllottedGeometry, DrawingParams);
+
+	const FVector2D TriggerThresholdRange(-1.0, 1.0);
+	const FVector2D DrawingAreaRange(DimensionSlot.Bottom, DimensionSlot.Top);
+
+	const double TriggerThresholdMappedValue = FMath::GetMappedRangeValueUnclamped(TriggerThresholdRange, DrawingAreaRange, TriggerThreshold);
+
+	const FVector2D StartLinePoint(0.0, TriggerThresholdMappedValue);
+	const FVector2D EndLinePoint(AllottedGeometry.GetLocalSize().X, TriggerThresholdMappedValue);
+
+	const TArray<FVector2D> TriggerLine({ StartLinePoint, EndLinePoint });
+
+	FSlateDrawElement::MakeLines(OutDrawElements, LayerId, AllottedGeometry.ToPaintGeometry(), TriggerLine, ESlateDrawEffect::None, LineColor, true, 1.0f);
+
+	return LayerId;
+}
+
+FVector2D STriggerThresholdLineWidget::ComputeDesiredSize(float) const
+{
+	return FVector2D();
+}
+
+void STriggerThresholdLineWidget::OnStyleUpdated(const FTriggerThresholdLineStyle UpdatedStyle)
+{
+	LineColor = UpdatedStyle.LineColor;
+}
+
+void STriggerThresholdLineWidget::SetTriggerThreshold(float InTriggerThreshold)
+{
+	TriggerThreshold = InTriggerThreshold;
+}
diff --git a/Engine/Plugins/Runtime/AudioWidgets/Source/AudioWidgets/Private/WaveformAudioSamplesDataProvider.cpp b/Engine/Plugins/Runtime/AudioWidgets/Source/AudioWidgets/Private/WaveformAudioSamplesDataProvider.cpp
new file mode 100644
--- /dev/null
+++ b/Engine/Plugins/Runtime/AudioWidgets/Source/AudioWidgets/Private/WaveformAudioSamplesDataProvider.cpp
@@ -0,0 +1,293 @@
+// Copyright Epic Games, Inc. All Rights Reserved.
+
+#include "WaveformAudioSamplesDataProvider.h"
+
+#include "AudioBusSubsystem.h"
+#include "AudioMixerDevice.h"
+#include "AudioWidgetsLog.h"
+#include "DSP/DeinterleaveView.h"
+
+namespace AudioWidgets
+{
+	namespace FWaveformAudioSamplesDataProviderPrivate
+	{
+		int32 GetTriggerThresholdBufferIndex(const EAudioOscilloscopeTriggerMode InTriggerMode, const float InTriggerThreshold, const TArray<float>& InAudioSamples)
+		{
+			for (int32 Index = 1; Index < InAudioSamples.Num(); ++Index)
+			{
+				if (InTriggerMode == EAudioOscilloscopeTriggerMode::Rising)
+				{
+					if (InAudioSamples[Index - 1] < InTriggerThreshold && InAudioSamples[Index] >= InTriggerThreshold)
+					{
+						return Index - 1;
+					}
+				}
+				else if (InTriggerMode == EAudioOscilloscopeTriggerMode::Falling)
+				{
+					if (InAudioSamples[Index - 1] > InTriggerThreshold && InAudioSamples[Index] <= InTriggerThreshold)
+					{
+						return Index - 1;
+					}
+				}
+			}
+
+			return -1;
+		};
+	}
+
+	FWaveformAudioSamplesDataProvider::FWaveformAudioSamplesDataProvider(UWorld* InWorld,
+		UAudioBus* InAudioBus,
+		const uint32 InNumChannelToProvide,
+		const float InTimeWindowMs,
+		const float InMaxTimeWindowMs,
+		const float InAnalysisPeriodMs)
+	{
+		if (!InWorld)
+		{
+			UE_LOG(LogAudioWidgets, Error, TEXT("World pointer needed to retrieve the AudioDevice is not available."));
+			return;
+		}
+
+		if (!InAudioBus)
+		{
+			UE_LOG(LogAudioWidgets, Error, TEXT("Unable to obtain audio samples for visualization without a valid audio bus."));
+			return;
+		}
+
+		MixerDevice = static_cast<Audio::FMixerDevice*>(InWorld->GetAudioDeviceRaw());
+
+		if (!MixerDevice)
+		{
+			UE_LOG(LogAudioWidgets, Error, TEXT("Unable to obtain audio samples for visualization without a valid audio mixer."));
+			return;
+		}
+
+		WorldPtr = InWorld;
+		AudioBus = InAudioBus;
+
+		NumChannelsToProvide = InNumChannelToProvide;
+		NumChannels          = InAudioBus->GetNumChannels();
+
+		SampleRate = static_cast<uint32>(MixerDevice->GetSampleRate());
+
+		// Init ticker handle
+		if (!TickerHandle.IsValid())
+		{
+			TickerHandle = FTSTicker::GetCoreTicker().AddTicker(FTickerDelegate::CreateRaw(this, &FWaveformAudioSamplesDataProvider::Tick), 0.0f);
+		}
+
+		// Init audio buffers
+		TimeWindowMaxTimeSamples = FMath::RoundToInt(InMaxTimeWindowMs / 1000.0f * NumChannelsToProvide * SampleRate);
+
+		AudioSamplesForView.Init(0.0f, TimeWindowMaxTimeSamples);
+		DataView = FFixedSampledSequenceView{ MakeArrayView(AudioSamplesForView.GetData(), AudioSamplesForView.Num()), NumChannelsToProvide, SampleRate };
+
+		AudioSamplesCircularBuffer.SetCapacity(TimeWindowMaxTimeSamples * 2); // Twice of the amount needed for display 
+
+		SetTimeWindow(InTimeWindowMs);
+		SetAnalysisPeriod(InAnalysisPeriodMs);
+	}
+
+
+	FWaveformAudioSamplesDataProvider::~FWaveformAudioSamplesDataProvider()
+	{
+		if (TickerHandle.IsValid())
+		{
+			FTSTicker::GetCoreTicker().RemoveTicker(TickerHandle);
+			TickerHandle.Reset();
+		}
+
+		StopProcessing();
+	}
+
+	void FWaveformAudioSamplesDataProvider::ResetAudioBuffers()
+	{
+		AudioSamplesCircularBuffer.Reset(TimeWindowMaxTimeSamples * 2);
+		AudioSamplesForView.Init(0.0f, AudioSamplesForView.Num());
+
+		AudioSamplesCircularBuffer.SetNum(TimeWindowSamples);
+	}
+
+	void FWaveformAudioSamplesDataProvider::StartProcessing()
+	{
+		if (bIsProcessing)
+		{
+			UE_LOG(LogAudioWidgets, Error, TEXT("StartProcessing invoked while having an active processing. Try call StopProcessing before calling StartProcessing."));
+			return;
+		}
+
+		check(WorldPtr);
+		check(AudioBus);
+		check(MixerDevice);
+
+		ResetAudioBuffers();
+
+		// Start the audio bus. This won't do anything if the bus is already started elsewhere.
+		const uint32 AudioBusId = AudioBus->GetUniqueID();
+
+		UAudioBusSubsystem* AudioBusSubsystem = MixerDevice->GetSubsystem<UAudioBusSubsystem>();
+		check(AudioBusSubsystem);
+
+		const Audio::FAudioBusKey AudioBusKey = Audio::FAudioBusKey(AudioBusId);
+
+		AudioBusSubsystem->StartAudioBus(AudioBusKey, NumChannels, false);
+
+		// Get an output patch for the audio bus
+		const uint32 NumFramesPerBufferToAnalyze = MixerDevice->GetNumOutputFrames();
+		PatchOutput = AudioBusSubsystem->AddPatchOutputForAudioBus(AudioBusKey, NumFramesPerBufferToAnalyze, NumChannels);
+
+		bIsProcessing = true;
+	}
+
+	void FWaveformAudioSamplesDataProvider::StopProcessing()
+	{
+		if (bIsProcessing)
+		{
+			PatchOutput   = nullptr;
+			bIsProcessing = false;
+		}
+	}
+
+	void FWaveformAudioSamplesDataProvider::SetChannelToAnalyze(const int32 InChannel)
+	{
+		ChannelIndexToAnalyze = InChannel - 1;
+		ResetAudioBuffers();
+	}
+
+	void FWaveformAudioSamplesDataProvider::SetTriggerMode(const EAudioOscilloscopeTriggerMode InTriggerMode)
+	{
+		TriggerMode   = InTriggerMode;
+		bHasTriggered = false;
+	}
+
+	void FWaveformAudioSamplesDataProvider::SetTriggerThreshold(const float InTriggerThreshold)
+	{
+		TriggerThreshold = InTriggerThreshold;
+		bHasTriggered    = false;
+	}
+
+	void FWaveformAudioSamplesDataProvider::SetTimeWindow(const float InTimeWindowMs)
+	{
+		TimeWindowSamples = FMath::RoundToInt(InTimeWindowMs / 1000.0f * NumChannelsToProvide * SampleRate);
+		DataView          = FFixedSampledSequenceView{ MakeArrayView(AudioSamplesForView.GetData(), AudioSamplesForView.Num()).Slice(0, TimeWindowSamples), NumChannelsToProvide, SampleRate };
+
+		if (NumChannelsToProvide == 1)
+		{
+			AudioSamplesCircularBuffer.SetNum(TimeWindowSamples);
+		}
+		else
+		{
+			// Adjust read pointer so it falls in a sample of the first channel
+			const uint32 ChannelsModulo              = TimeWindowSamples % NumChannelsToProvide;
+			const uint32 TimeWindowSamplesWithOffset = TimeWindowSamples - ChannelsModulo;
+
+			AudioSamplesCircularBuffer.SetNum(TimeWindowSamplesWithOffset);
+		}
+	}
+
+	void FWaveformAudioSamplesDataProvider::SetAnalysisPeriod(const float InAnalysisPeriodMs)
+	{
+		AnalysisPeriodSamples = InAnalysisPeriodMs * 0.001f * SampleRate;
+		AudioSamplesCircularBuffer.SetNum(TimeWindowSamples);
+	}
+
+	void FWaveformAudioSamplesDataProvider::PushAudioSamplesToCircularBuffer()
+	{
+		if (const int32 NumSamplesAvailable = PatchOutput->GetNumSamplesAvailable();
+			NumSamplesAvailable > 0)
+		{
+			TempAudioBuffer.Reset();
+			TempAudioBuffer.AddUninitialized(NumSamplesAvailable);
+
+			PatchOutput->PopAudio(TempAudioBuffer.GetData(), NumSamplesAvailable, false);
+
+			if (NumChannelsToProvide == 1)
+			{
+				// Deinterleave so we push into the circular buffer the channel of interest
+				TArray<float> ArrayToFill;
+				Audio::TAutoDeinterleaveView<float> DeinterleaveView(TempAudioBuffer, ArrayToFill, NumChannels);
+
+				for (auto Channel : DeinterleaveView)
+				{
+					if (Channel.ChannelIndex == ChannelIndexToAnalyze)
+					{
+						NumSamplesPushedToCircularBuffer += AudioSamplesCircularBuffer.Push(Channel.Values.GetData(), Channel.Values.Num());
+
+						// Check if a trigger has occurred
+						if (TriggerMode != EAudioOscilloscopeTriggerMode::None)
+						{
+							if (const int32 TriggerThresholdBufferIndex = FWaveformAudioSamplesDataProviderPrivate::GetTriggerThresholdBufferIndex(TriggerMode, TriggerThreshold, Channel.Values);
+								TriggerThresholdBufferIndex != -1 && !bHasTriggered)
+							{
+								const int32 TriggerScreenPosition = 0; // TODO alex.perez: tweakable param? can be from 0 to TimeWindowSamples - 1;
+
+								// Move the read pointer at the trigger point
+								AudioSamplesCircularBuffer.SetNum(NumSamplesPushedToCircularBuffer - TriggerThresholdBufferIndex + TriggerScreenPosition);
+
+								bHasTriggered = true;
+							}
+						}
+					}
+				}
+			}
+			else
+			{
+				NumSamplesPushedToCircularBuffer += AudioSamplesCircularBuffer.Push(TempAudioBuffer.GetData(), TempAudioBuffer.Num());
+			}
+		}
+	}
+
+	bool FWaveformAudioSamplesDataProvider::Tick(float DeltaTime)
+	{
+		TRACE_CPUPROFILER_EVENT_SCOPE(FWaveformAudioSamplesDataProvider::Tick);
+
+		if (PatchOutput.IsValid())
+		{
+			PushAudioSamplesToCircularBuffer();
+
+			const uint32 NumAvailableSamplesToReadCircularBuffer = AudioSamplesCircularBuffer.Num();
+
+			if (TriggerMode == EAudioOscilloscopeTriggerMode::None || !bHasTriggered)
+			{
+				if (NumAvailableSamplesToReadCircularBuffer >= TimeWindowSamples)
+				{
+					// Obtain the audio samples we want to display
+					AudioSamplesCircularBuffer.Peek(AudioSamplesForView.GetData(), TimeWindowSamples);
+
+					// Advance the read pointer
+					AudioSamplesCircularBuffer.Pop(AnalysisPeriodSamples * NumChannelsToProvide);
+				}
+			}
+			else
+			{
+				// Keep refreshing the trigger view with new samples until we reach TimeWindowSamples
+				if (NumAvailableSamplesToReadCircularBuffer < TimeWindowSamples)
+				{
+					AudioSamplesCircularBuffer.Peek(AudioSamplesForView.GetData(), NumAvailableSamplesToReadCircularBuffer);
+				}
+				else
+				{
+					AudioSamplesCircularBuffer.Pop(AudioSamplesForView.GetData(), NumAvailableSamplesToReadCircularBuffer);
+					bHasTriggered = false;
+				}
+			}
+
+			// Send data if we have reached the analysis period time
+			if (NumSamplesPushedToCircularBuffer >= AnalysisPeriodSamples * NumChannelsToProvide)
+			{
+				RequestSequenceView(TRange<double>::Inclusive(0, 1));
+				NumSamplesPushedToCircularBuffer = 0;
+			}
+		}
+
+		return true;
+	}
+
+	FFixedSampledSequenceView FWaveformAudioSamplesDataProvider::RequestSequenceView(const TRange<double> DataRatioRange)
+	{
+		// Broadcast audio samples data view
+		OnDataViewGenerated.Broadcast(DataView, /*FirstRenderedSample*/ 0);
+
+		return DataView;
+	}
+} // namespace AudioWidgets
diff --git a/Engine/Plugins/Runtime/AudioWidgets/Source/AudioWidgets/Public/AudioOscilloscope.h b/Engine/Plugins/Runtime/AudioWidgets/Source/AudioWidgets/Public/AudioOscilloscope.h
new file mode 100644
--- /dev/null
+++ b/Engine/Plugins/Runtime/AudioWidgets/Source/AudioWidgets/Public/AudioOscilloscope.h
@@ -0,0 +1,42 @@
+// Copyright Epic Games, Inc. All Rights Reserved.
+
+#pragma once
+
+#include "AudioOscilloscopePanelStyle.h"
+#include "AudioWidgetsEnums.h"
+#include "SAudioOscilloscopePanelWidget.h"
+#include "Sound/AudioBus.h"
+#include "UObject/StrongObjectPtr.h"
+#include "WaveformAudioSamplesDataProvider.h"
+
+class SAudioOscilloscopePanelWidget;
+
+namespace AudioWidgets
+{
+	class FWaveformAudioSamplesDataProvider;
+
+	class AUDIOWIDGETS_API FAudioOscilloscope
+	{
+	public:
+		FAudioOscilloscope(UWorld* InWorld, 
+			const uint32 InNumChannels, 
+			const float InTimeWindowMs, 
+			const float InMaxTimeWindowMs, 
+			const float InAnalysisPeriodMs, 
+			const EAudioPanelLayoutType InPanelLayoutType);
+
+		void StartProcessing();
+		void StopProcessing();
+
+		UAudioBus* GetAudioBus() const;
+		TSharedRef<SWidget> GetPanelWidget() const;
+
+	private:
+		FAudioOscilloscopePanelStyle OscilloscopePanelStyle;
+
+		TSharedPtr<FWaveformAudioSamplesDataProvider> AudioSamplesDataProvider = nullptr;
+		TSharedPtr<SAudioOscilloscopePanelWidget> OscilloscopePanelWidget      = nullptr;
+
+		TStrongObjectPtr<UAudioBus> AudioBus = nullptr;
+	};
+} // namespace AudioWidgets
diff --git a/Engine/Plugins/Runtime/AudioWidgets/Source/AudioWidgets/Public/AudioOscilloscopeEnums.h b/Engine/Plugins/Runtime/AudioWidgets/Source/AudioWidgets/Public/AudioOscilloscopeEnums.h
new file mode 100644
--- /dev/null
+++ b/Engine/Plugins/Runtime/AudioWidgets/Source/AudioWidgets/Public/AudioOscilloscopeEnums.h
@@ -0,0 +1,13 @@
+// Copyright Epic Games, Inc. All Rights Reserved.
+
+#pragma once
+
+#include "AudioOscilloscopeEnums.generated.h"
+
+UENUM(BlueprintType)
+enum class EAudioOscilloscopeTriggerMode : uint8
+{
+	None    UMETA(DisplayName = "None"),
+	Rising  UMETA(DisplayName = "Rising"),
+	Falling UMETA(DisplayName = "Falling")
+};
diff --git a/Engine/Plugins/Runtime/AudioWidgets/Source/AudioWidgets/Public/AudioOscilloscopePanelStyle.h b/Engine/Plugins/Runtime/AudioWidgets/Source/AudioWidgets/Public/AudioOscilloscopePanelStyle.h
new file mode 100644
--- /dev/null
+++ b/Engine/Plugins/Runtime/AudioWidgets/Source/AudioWidgets/Public/AudioOscilloscopePanelStyle.h
@@ -0,0 +1,85 @@
+// Copyright Epic Games, Inc. All Rights Reserved.
+
+#pragma once
+
+#include "AudioWidgetsSlateTypes.h"
+#include "Fonts/SlateFontInfo.h"
+#include "Styling/SlateBrush.h"
+#include "Styling/SlateWidgetStyle.h"
+#include "Styling/StyleDefaults.h"
+#include "TriggerThresholdLineStyle.h"
+
+#include "AudioOscilloscopePanelStyle.generated.h"
+
+DECLARE_MULTICAST_DELEGATE_OneParam(FOnNewTimeRulerStyle,        FFixedSampleSequenceRulerStyle)
+DECLARE_MULTICAST_DELEGATE_OneParam(FOnNewValueGridOverlayStyle, FSampledSequenceValueGridOverlayStyle)
+DECLARE_MULTICAST_DELEGATE_OneParam(FOnNewWaveformViewerStyle,   FSampledSequenceViewerStyle)
+
+/**
+* Represents the appearance of an SAudioOscilloscopePanelWidget
+*/
+USTRUCT(BlueprintType)
+struct AUDIOWIDGETS_API FAudioOscilloscopePanelStyle : public FSlateWidgetStyle
+{
+	GENERATED_USTRUCT_BODY()
+
+	FAudioOscilloscopePanelStyle()
+	{
+		TimeRulerStyle = FFixedSampleSequenceRulerStyle()
+		.SetHandleColor(FLinearColor(0.0f, 0.0f, 0.0f, 0.0f))
+		.SetTicksColor(FLinearColor(1.0f, 1.0f, 1.0f, 1.0f))
+		.SetTicksTextColor(FLinearColor(1.0f, 1.0f, 1.0f, 1.0f))
+		.SetHandleColor(FLinearColor(0.0f, 0.0f, 0.0f, 0.0f))
+		.SetFontSize(10);
+
+		ValueGridStyle = FSampledSequenceValueGridOverlayStyle()
+		.SetGridColor(FLinearColor(0.05f, 0.05f, 0.05f, 0.5f))
+		.SetGridThickness(1.0f)
+		.SetLabelTextColor(FLinearColor(0.442708f, 0.442708f, 0.442708f, 1.0f))
+		.SetLabelTextFontSize(10.0f);
+
+		WaveViewerStyle = FSampledSequenceViewerStyle()
+		.SetSequenceColor(FLinearColor(0.1f, 1.0f, 0.55f, 1.0f))
+		.SetBackgroundColor(FLinearColor(0.01f, 0.01f, 0.01f, 1.0f))
+		.SetSequenceLineThickness(2.0f)
+		.SetSampleMarkersSize(0.0f)
+		.SetMajorGridLineColor(FLinearColor(0.1f, 0.1f, 0.1f, 1.0f))
+		.SetMinorGridLineColor(FLinearColor(0.05f, 0.05f, 0.05f, 1.0f))
+		.SetZeroCrossingLineColor(FLinearColor(0.3f, 0.3f, 0.3f, 1.0f))
+		.SetZeroCrossingLineThickness(1.0f);
+
+		TriggerThresholdLineStyle = FTriggerThresholdLineStyle()
+		.SetLineColor(FLinearColor(1.0f, 0.91f, 0.34f));
+	}
+
+	virtual ~FAudioOscilloscopePanelStyle() {}
+
+	virtual const FName GetTypeName() const override { return TypeName; }
+	static const FAudioOscilloscopePanelStyle& GetDefault() { static FAudioOscilloscopePanelStyle Default; return Default; }
+
+	/** The time rule style. */
+	UPROPERTY(EditAnywhere, BlueprintReadWrite, Category = Appearance)
+	FFixedSampleSequenceRulerStyle TimeRulerStyle;
+	FAudioOscilloscopePanelStyle& SetTimeRulerStyle(const FFixedSampleSequenceRulerStyle& InTimeRulerStyle) { TimeRulerStyle = InTimeRulerStyle; return *this; }
+
+	/** The value grid style. */
+	UPROPERTY(EditAnywhere, BlueprintReadWrite, Category = Appearance)
+	FSampledSequenceValueGridOverlayStyle ValueGridStyle;
+	FAudioOscilloscopePanelStyle& SetValueGridStyle(const FSampledSequenceValueGridOverlayStyle& InValueGridStyle) { ValueGridStyle = InValueGridStyle; return *this; }
+
+	/** The waveform view style. */
+	UPROPERTY(EditAnywhere, BlueprintReadWrite, Category = Appearance)
+	FSampledSequenceViewerStyle WaveViewerStyle;
+	FAudioOscilloscopePanelStyle& SetWaveViewerStyle(const FSampledSequenceViewerStyle& InWaveViewerStyle) { WaveViewerStyle = InWaveViewerStyle; return *this; }
+
+	/** The triggerthreshold line style. */
+	UPROPERTY(EditAnywhere, BlueprintReadWrite, Category = Appearance)
+	FTriggerThresholdLineStyle TriggerThresholdLineStyle;
+	FAudioOscilloscopePanelStyle& SetTriggerThresholdLineStyle(const FTriggerThresholdLineStyle& InTriggerThresholdLineStyle) { TriggerThresholdLineStyle = InTriggerThresholdLineStyle; return *this; }
+
+	inline static FName TypeName = FName("FAudioOscilloscopePanelStyle");
+
+	inline static FOnNewTimeRulerStyle        OnNewTimeRulerStyle;
+	inline static FOnNewValueGridOverlayStyle OnNewValueGridOverlayStyle;
+	inline static FOnNewWaveformViewerStyle   OnNewWaveformViewerStyle;
+};
diff --git a/Engine/Plugins/Runtime/AudioWidgets/Source/AudioWidgets/Public/AudioOscilloscopeUMG.h b/Engine/Plugins/Runtime/AudioWidgets/Source/AudioWidgets/Public/AudioOscilloscopeUMG.h
new file mode 100644
--- /dev/null
+++ b/Engine/Plugins/Runtime/AudioWidgets/Source/AudioWidgets/Public/AudioOscilloscopeUMG.h
@@ -0,0 +1,118 @@
+// Copyright Epic Games, Inc. All Rights Reserved.
+
+#pragma once
+
+#include "AudioOscilloscopePanelStyle.h"
+#include "AudioWidgetsEnums.h"
+#include "Components/Widget.h"
+#include "FixedSampledSequenceView.h"
+#include "SampledSequenceDisplayUnit.h"
+#include "SAudioOscilloscopePanelWidget.h"
+#include "Sound/AudioBus.h"
+
+#include "AudioOscilloscopeUMG.generated.h"
+
+namespace AudioWidgets { class FWaveformAudioSamplesDataProvider; }
+class UAudioBus;
+class UWorld;
+
+/**
+ * An oscilloscope UMG widget.
+ *
+ * Supports displaying waveforms from incoming audio samples.
+ * 
+ */
+UCLASS()
+class AUDIOWIDGETS_API UAudioOscilloscope: public UWidget
+{
+	GENERATED_UCLASS_BODY()
+
+public:
+	DECLARE_DYNAMIC_DELEGATE_RetVal(TArray<float>, FGetOscilloscopeAudioSamples);
+
+	// UWidget overrides
+	virtual void SynchronizeProperties() override;
+	virtual void ReleaseSlateResources(bool bReleaseChildren) override;
+
+#if WITH_EDITOR
+	virtual const FText GetPaletteCategory() override;
+#endif
+
+	/** Starts the oscilloscope processing. */
+	UFUNCTION(BlueprintCallable, Category = "Behavior")
+	void StartProcessing();
+
+	/** Stops the oscilloscope processing. */
+	UFUNCTION(BlueprintCallable, Category = "Behavior")
+	void StopProcessing();
+
+	/** The oscilloscope panel style */
+	UPROPERTY(EditAnywhere, BlueprintReadWrite, Category="Style", meta=(DisplayName="Style"))
+	FAudioOscilloscopePanelStyle OscilloscopeStyle;
+
+	/** The audio bus used to obtain audio samples for the oscilloscope */
+	UPROPERTY(EditAnywhere, BlueprintReadWrite, Category = "Oscilloscope Values", meta = (DesignerRebuild = "True"))
+	TObjectPtr<UAudioBus> AudioBus = nullptr;
+
+	/** The time window in milliseconds. */
+	UPROPERTY(EditAnywhere, BlueprintReadWrite, Category = "Oscilloscope Values", meta = (UIMin = 10, UIMax = 5000, ClampMin = 10, ClampMax = 5000))
+	float TimeWindowMs = 10.0f;
+
+	/** The analysis period in milliseconds. */
+	UPROPERTY(EditAnywhere, BlueprintReadWrite, Category = "Oscilloscope Values", meta = (UIMin = 10, UIMax = 1000, ClampMin = 10, ClampMax = 1000))
+	float AnalysisPeriodMs = 10.0f;
+
+	/** Show/Hide the time grid. */
+	UPROPERTY(EditAnywhere, BlueprintReadWrite, Category = "Oscilloscope Values")
+	bool bShowTimeGrid = true;
+
+	/** Define the time grid labels unit. */
+	UPROPERTY(EditAnywhere, BlueprintReadWrite, Category = "Oscilloscope Values")
+	EXAxisLabelsUnit TimeGridLabelsUnit = EXAxisLabelsUnit::Samples;
+
+	/** Show/Hide the amplitude grid. */
+	UPROPERTY(EditAnywhere, BlueprintReadWrite, Category = "Oscilloscope Values")
+	bool bShowAmplitudeGrid = true;
+
+	/** Show/Hide the amplitude labels. */
+	UPROPERTY(EditAnywhere, BlueprintReadWrite, Category = "Oscilloscope Values")
+	bool bShowAmplitudeLabels = true;
+
+	/** Define the amplitude grid labels unit. */
+	UPROPERTY(EditAnywhere, BlueprintReadWrite, Category = "Oscilloscope Values")
+	EYAxisLabelsUnit AmplitudeGridLabelsUnit = EYAxisLabelsUnit::Linear;
+
+	/** Show/Hide the trigger threshold line. */
+	UPROPERTY(EditAnywhere, BlueprintReadWrite, Category = "Oscilloscope Values")
+	bool bShowTriggerThresholdLine = false;
+
+	/** The trigger threshold position in the Y axis. */
+	UPROPERTY(EditAnywhere, BlueprintReadWrite, Category = "Oscilloscope Values", meta = (UIMin = -1, UIMax = 1, ClampMin = -1, ClampMax = 1))
+	float TriggerThreshold = 0.0f;
+
+	/** Show/Hide advanced panel layout. */
+	UPROPERTY(EditAnywhere, BlueprintReadWrite, Category = "Oscilloscope Values", meta = (DesignerRebuild = "True"))
+	EAudioPanelLayoutType PanelLayoutType = EAudioPanelLayoutType::Basic;
+
+private:
+	void CreateDummyOscilloscopeWidget();
+	void CreateDataProvider();
+	void CreateOscilloscopeWidget();
+
+	// UWidget overrides
+	virtual TSharedRef<SWidget> RebuildWidget() override;
+	
+	// The underlying audio samples data provider
+	TSharedPtr<AudioWidgets::FWaveformAudioSamplesDataProvider> AudioSamplesDataProvider;
+
+	// Native Slate Widget
+	TSharedPtr<SAudioOscilloscopePanelWidget> OscilloscopePanelWidget;
+
+	// Dummy waveform data to display if audio bus is not set
+	static constexpr uint32 DummySampleRate   = 48000;
+	static constexpr int32 DummyMaxNumSamples = DummySampleRate * 5;
+	static constexpr int32 DummyNumChannels   = 1;
+
+	TArray<float> DummyAudioSamples;
+	FFixedSampledSequenceView DummyDataView;
+};
diff --git a/Engine/Plugins/Runtime/AudioWidgets/Source/AudioWidgets/Public/AudioVectorscope.h b/Engine/Plugins/Runtime/AudioWidgets/Source/AudioWidgets/Public/AudioVectorscope.h
new file mode 100644
--- /dev/null
+++ b/Engine/Plugins/Runtime/AudioWidgets/Source/AudioWidgets/Public/AudioVectorscope.h
@@ -0,0 +1,41 @@
+// Copyright Epic Games, Inc. All Rights Reserved.
+
+#pragma once
+
+#include "AudioVectorscopePanelStyle.h"
+#include "AudioWidgetsEnums.h"
+#include "Sound/AudioBus.h"
+#include "UObject/StrongObjectPtr.h"
+#include "WaveformAudioSamplesDataProvider.h"
+
+class SAudioVectorscopePanelWidget;
+
+namespace AudioWidgets
+{
+	class FWaveformAudioSamplesDataProvider;
+
+	class AUDIOWIDGETS_API FAudioVectorscope
+	{
+	public:
+		FAudioVectorscope(UWorld* InWorld, 
+			const uint32 InNumChannels, 
+			const float InTimeWindowMs, 
+			const float InMaxTimeWindowMs, 
+			const float InAnalysisPeriodMs, 
+			const EAudioPanelLayoutType InPanelLayoutType);
+
+		void StartProcessing();
+		void StopProcessing();
+
+		UAudioBus* GetAudioBus() const;
+		TSharedRef<SWidget> GetPanelWidget() const;
+
+	private:
+		FAudioVectorscopePanelStyle VectorscopePanelStyle;
+
+		TSharedPtr<FWaveformAudioSamplesDataProvider> AudioSamplesDataProvider = nullptr;
+		TSharedPtr<SAudioVectorscopePanelWidget> VectorscopePanelWidget        = nullptr;
+
+		TStrongObjectPtr<UAudioBus> AudioBus = nullptr;
+	};
+} // namespace AudioWidgets
diff --git a/Engine/Plugins/Runtime/AudioWidgets/Source/AudioWidgets/Public/AudioVectorscopePanelStyle.h b/Engine/Plugins/Runtime/AudioWidgets/Source/AudioWidgets/Public/AudioVectorscopePanelStyle.h
new file mode 100644
--- /dev/null
+++ b/Engine/Plugins/Runtime/AudioWidgets/Source/AudioWidgets/Public/AudioVectorscopePanelStyle.h
@@ -0,0 +1,53 @@
+// Copyright Epic Games, Inc. All Rights Reserved.
+
+#pragma once
+
+#include "AudioWidgetsSlateTypes.h"
+#include "SampledSequenceVectorViewerStyle.h"
+#include "Styling/SlateWidgetStyle.h"
+
+#include "AudioVectorscopePanelStyle.generated.h"
+
+DECLARE_MULTICAST_DELEGATE_OneParam(FOnNewValueGridOverlayStyle, FSampledSequenceValueGridOverlayStyle)
+DECLARE_MULTICAST_DELEGATE_OneParam(FOnNewVectorViewerStyle,     FSampledSequenceVectorViewerStyle)
+
+/**
+* Represents the appearance of an SAudioVectorscopePanelWidget
+*/
+USTRUCT(BlueprintType)
+struct AUDIOWIDGETS_API FAudioVectorscopePanelStyle : public FSlateWidgetStyle
+{
+	GENERATED_USTRUCT_BODY()
+
+		FAudioVectorscopePanelStyle()
+	{
+		ValueGridStyle = FSampledSequenceValueGridOverlayStyle()
+		.SetGridColor(FLinearColor(0.05f, 0.05f, 0.05f, 0.5f))
+		.SetGridThickness(1.0f);
+
+		VectorViewerStyle = FSampledSequenceVectorViewerStyle()
+		.SetBackgroundColor(FLinearColor(0.00273f, 0.00478f, 0.00368f, 1.0f))
+		.SetLineColor(FLinearColor(0.0075f, 0.57758f, 0.05448f, 1.0f))
+		.SetLineThickness(2.5f);
+	}
+
+	virtual ~FAudioVectorscopePanelStyle() {}
+
+	virtual const FName GetTypeName() const override { return TypeName; }
+	static const FAudioVectorscopePanelStyle& GetDefault() { static FAudioVectorscopePanelStyle Default; return Default; }
+
+	/** The value grid style. */
+	UPROPERTY(EditAnywhere, BlueprintReadWrite, Category = Appearance)
+	FSampledSequenceValueGridOverlayStyle ValueGridStyle;
+	FAudioVectorscopePanelStyle& SetValueGridStyle(const FSampledSequenceValueGridOverlayStyle& InValueGridStyle) { ValueGridStyle = InValueGridStyle; return *this; }
+
+	/** The vector view style. */
+	UPROPERTY(EditAnywhere, BlueprintReadWrite, Category = Appearance)
+	FSampledSequenceVectorViewerStyle VectorViewerStyle;
+	FAudioVectorscopePanelStyle& SetVectorViewerStyle(const FSampledSequenceVectorViewerStyle& InVectorViewerStyle) { VectorViewerStyle = InVectorViewerStyle; return *this; }
+
+	inline static FName TypeName = FName("FAudioVectorscopePanelStyle");
+
+	inline static FOnNewValueGridOverlayStyle OnNewValueGridOverlayStyle;
+	inline static FOnNewVectorViewerStyle     OnNewVectorViewerStyle;
+};
diff --git a/Engine/Plugins/Runtime/AudioWidgets/Source/AudioWidgets/Public/AudioVectorscopeUMG.h b/Engine/Plugins/Runtime/AudioWidgets/Source/AudioWidgets/Public/AudioVectorscopeUMG.h
new file mode 100644
--- /dev/null
+++ b/Engine/Plugins/Runtime/AudioWidgets/Source/AudioWidgets/Public/AudioVectorscopeUMG.h
@@ -0,0 +1,93 @@
+// Copyright Epic Games, Inc. All Rights Reserved.
+
+#pragma once
+
+#include "AudioVectorscopePanelStyle.h"
+#include "AudioWidgetsEnums.h"
+#include "Components/Widget.h"
+#include "FixedSampledSequenceView.h"
+#include "Sound/AudioBus.h"
+
+#include "AudioVectorscopeUMG.generated.h"
+
+namespace AudioWidgets { class FWaveformAudioSamplesDataProvider; }
+class SAudioVectorscopePanelWidget;
+
+/**
+ * A vectorscope UMG widget.
+ *
+ * Supports displaying waveforms in 2D (Left channel X axis, Right channel Y axis) from incoming audio samples.
+ * 
+ */
+UCLASS()
+class AUDIOWIDGETS_API UAudioVectorscope: public UWidget
+{
+	GENERATED_UCLASS_BODY()
+
+public:
+	DECLARE_DYNAMIC_DELEGATE_RetVal(TArray<float>, FGetVectorscopeAudioSamples);
+
+	// UWidget overrides
+	virtual void SynchronizeProperties() override;
+	virtual void ReleaseSlateResources(bool bReleaseChildren) override;
+
+#if WITH_EDITOR
+	virtual const FText GetPaletteCategory() override;
+#endif
+
+	/** Starts the vectorscope processing. */
+	UFUNCTION(BlueprintCallable, Category = "Behavior")
+	void StartProcessing();
+
+	/** Stops the vectorscope processing. */
+	UFUNCTION(BlueprintCallable, Category = "Behavior")
+	void StopProcessing();
+
+	/** The vectorscope panel style */
+	UPROPERTY(EditAnywhere, BlueprintReadWrite, Category="Style", meta=(DisplayName="Style"))
+	FAudioVectorscopePanelStyle VectorscopeStyle;
+
+	/** The audio bus used to obtain audio samples for the vectorscope */
+	UPROPERTY(EditAnywhere, BlueprintReadWrite, Category = "Vectorscope Values", meta = (DesignerRebuild = "True"))
+	TObjectPtr<UAudioBus> AudioBus = nullptr;
+
+	/** Show/Hide the vectorscope grid. */
+	UPROPERTY(EditAnywhere, BlueprintReadWrite, Category = "Vectorscope Values")
+	bool bShowGrid = true;
+
+	/** The number of grid divisions. */
+	UPROPERTY(EditAnywhere, BlueprintReadWrite, Category = "Vectorscope Values", meta = (UIMin = 1, UIMax = 6, ClampMin = 1, ClampMax = 6))
+	int32 GridDivisions = 2;
+
+	/** For how long the audio samples should persist in the screen (in milliseconds). */
+	UPROPERTY(EditAnywhere, BlueprintReadWrite, Category = "Vectorscope Values", meta = (UIMin = 10.0, UIMax = 100.0, ClampMin = 10.0, ClampMax = 100.0))
+	float DisplayPersistenceMs = 60.0f;
+
+	/** The scale for the displayed audio samples. */
+	UPROPERTY(EditAnywhere, BlueprintReadWrite, Category = "Vectorscope Values", meta = (UIMin = 0.0, UIMax = 1.0, ClampMin = 0.0, ClampMax = 1.0))
+	float Scale = 1.0f;
+
+	/** Show/Hide advanced panel layout. */
+	UPROPERTY(EditAnywhere, BlueprintReadWrite, Category = "Vectorscope Values", meta = (DesignerRebuild = "True"))
+	EAudioPanelLayoutType PanelLayoutType = EAudioPanelLayoutType::Basic;
+
+private:
+	void CreateDummyVectorscopeWidget();
+	void CreateDataProvider();
+	void CreateVectorscopeWidget();
+
+	// UWidget overrides
+	virtual TSharedRef<SWidget> RebuildWidget() override;
+
+	static constexpr float AnalysisPeriodMs = 10.0f;
+
+	// The underlying audio samples data provider
+	TSharedPtr<AudioWidgets::FWaveformAudioSamplesDataProvider> AudioSamplesDataProvider;
+
+	// Native Slate Widget
+	TSharedPtr<SAudioVectorscopePanelWidget> VectorscopePanelWidget;
+
+	// Dummy waveform data to display if audio bus is not set
+	TArray<float> DummyAudioSamples;
+	FFixedSampledSequenceView DummyDataView;
+};
diff --git a/Engine/Plugins/Runtime/AudioWidgets/Source/AudioWidgets/Public/AudioWidgetsEnums.h b/Engine/Plugins/Runtime/AudioWidgets/Source/AudioWidgets/Public/AudioWidgetsEnums.h
new file mode 100644
--- /dev/null
+++ b/Engine/Plugins/Runtime/AudioWidgets/Source/AudioWidgets/Public/AudioWidgetsEnums.h
@@ -0,0 +1,12 @@
+// Copyright Epic Games, Inc. All Rights Reserved.
+
+#pragma once
+
+#include "AudioWidgetsEnums.generated.h"
+
+UENUM(BlueprintType)
+enum class EAudioPanelLayoutType : uint8
+{
+	Basic    UMETA(DisplayName = "Basic"),
+	Advanced UMETA(DisplayName = "Advanced")
+};
diff --git a/Engine/Plugins/Runtime/AudioWidgets/Source/AudioWidgets/Public/AudioWidgetsLog.h b/Engine/Plugins/Runtime/AudioWidgets/Source/AudioWidgets/Public/AudioWidgetsLog.h
new file mode 100644
--- /dev/null
+++ b/Engine/Plugins/Runtime/AudioWidgets/Source/AudioWidgets/Public/AudioWidgetsLog.h
@@ -0,0 +1,7 @@
+// Copyright 2023, SweejTech Ltd. All Rights Reserved.
+
+#pragma once
+
+#include "Logging/LogMacros.h"
+
+AUDIOWIDGETS_API DECLARE_LOG_CATEGORY_EXTERN(LogAudioWidgets, Log, All);
diff --git a/Engine/Plugins/Runtime/AudioWidgets/Source/AudioWidgets/Public/SAudioOscilloscopePanelWidget.h b/Engine/Plugins/Runtime/AudioWidgets/Source/AudioWidgets/Public/SAudioOscilloscopePanelWidget.h
new file mode 100644
--- /dev/null
+++ b/Engine/Plugins/Runtime/AudioWidgets/Source/AudioWidgets/Public/SAudioOscilloscopePanelWidget.h
@@ -0,0 +1,190 @@
+// Copyright Epic Games, Inc. All Rights Reserved.
+
+#pragma once
+
+#include "AudioOscilloscopeEnums.h"
+#include "AudioOscilloscopePanelStyle.h"
+#include "AudioWidgetsEnums.h"
+#include "AudioWidgetsSlateTypes.h"
+#include "AudioWidgetsStyle.h"
+#include "IFixedSampledSequenceViewReceiver.h"
+#include "SampledSequenceDisplayUnit.h"
+#include "SampledSequenceDrawingUtils.h"
+#include "SSampledSequenceValueGridOverlay.h"
+#include "STriggerThresholdLineWidget.h"
+#include "TriggerThresholdLineStyle.h"
+#include "Widgets/Input/SComboBox.h"
+#include "Widgets/SCompoundWidget.h"
+
+class FFixedSampledSequenceGridData;
+class SAudioRadialSlider;
+class SBorder;
+class SFixedSampledSequenceRuler;
+class SFixedSampledSequenceViewer;
+class STriggerThresholdLineWidget;
+
+DECLARE_MULTICAST_DELEGATE_OneParam(FOnSelectedChannelChanged,  const int32 /*InChannel*/)
+DECLARE_MULTICAST_DELEGATE_OneParam(FOnTriggerModeChanged,      const EAudioOscilloscopeTriggerMode /*InTriggerMode*/)
+DECLARE_MULTICAST_DELEGATE_OneParam(FOnTriggerThresholdChanged, const float /*InTriggerThreshold*/)
+DECLARE_MULTICAST_DELEGATE_OneParam(FOnTimeWindowValueChanged,  const float /*InTimeWindowMs*/)
+DECLARE_MULTICAST_DELEGATE_OneParam(FOnAnalysisPeriodChanged,   const float /*InAnalysisPeriodMs*/)
+
+UENUM(BlueprintType)
+enum class EXAxisLabelsUnit : uint8
+{
+	Samples UMETA(DisplayName = "Samples"),
+	Seconds UMETA(DisplayName = "Seconds")
+};
+
+UENUM(BlueprintType)
+enum class EYAxisLabelsUnit : uint8
+{
+	Linear UMETA(DisplayName = "Linear"),
+	Db     UMETA(DisplayName = "dB")
+};
+
+class AUDIOWIDGETS_API SAudioOscilloscopePanelWidget : public SCompoundWidget, public IFixedSampledSequenceViewReceiver
+{
+public:
+	SLATE_BEGIN_ARGS(SAudioOscilloscopePanelWidget)
+		: _HideSequenceGrid(false)
+		, _HideSequenceRuler(false)
+		, _HideValueGrid(false)
+		, _HideTriggerThresholdLine(true)
+		, _YAxisLabelsUnit(EYAxisLabelsUnit::Linear)
+		, _ValueGridMaxDivisionParameter(2)
+		, _SequenceRulerDisplayUnit(EXAxisLabelsUnit::Samples)
+		, _PanelLayoutType(EAudioPanelLayoutType::Basic)
+		, _PanelStyle(&FAudioOscilloscopePanelStyle::GetDefault())
+	{}
+		
+		SLATE_DEFAULT_SLOT(FArguments, InArgs)
+
+		/** Whether the sequenbce grid should be drawn or not */
+		SLATE_ATTRIBUTE(bool, HideSequenceGrid)
+
+		/** Whether the sequence ruler should be drawn or not */
+		SLATE_ATTRIBUTE(bool, HideSequenceRuler)
+
+		/** Whether the value grid should be drawn or not */
+		SLATE_ATTRIBUTE(bool, HideValueGrid)
+
+		/** Whether the trigger threshold line should be drawn or not */
+		SLATE_ATTRIBUTE(bool, HideTriggerThresholdLine)
+
+		/** Whether the value grid values are linear or in dB */
+		SLATE_ATTRIBUTE(EYAxisLabelsUnit, YAxisLabelsUnit)
+
+		/** Maximum number of divisions in the value grid */
+		SLATE_ARGUMENT(uint32, ValueGridMaxDivisionParameter)
+
+		/** Sequence ruler display unit */
+		SLATE_ARGUMENT(EXAxisLabelsUnit, SequenceRulerDisplayUnit)
+
+		/** If we want to set the basic or advanced layout */
+		SLATE_ARGUMENT(EAudioPanelLayoutType, PanelLayoutType)
+
+		/** Oscilloscope Panel widget style */
+		SLATE_STYLE_ARGUMENT(FAudioOscilloscopePanelStyle, PanelStyle)
+
+	SLATE_END_ARGS()
+
+	void Construct(const FArguments& InArgs, const FFixedSampledSequenceView& InData, const int32 InNumChannels);
+
+	// SWidget methods overrides
+	virtual void Tick(const FGeometry& AllottedGeometry, const double InCurrentTime, const float InDeltaTime) override;
+
+	// IFixedSampledSequenceViewReceiver Interface
+	virtual void ReceiveSequenceView(const FFixedSampledSequenceView InData, const uint32 FirstSampleIndex) override;
+
+	void SetSequenceRulerDisplayUnit(const EXAxisLabelsUnit InDisplayUnit);
+	void SetXAxisGridVisibility(const bool InbIsVisible);
+
+	void SetValueGridOverlayDisplayUnit(const EYAxisLabelsUnit InDisplayUnit);
+	void SetYAxisGridVisibility(const bool InbIsVisible);
+	void SetYAxisLabelsVisibility(const bool InbIsVisible);
+
+	void SetTriggerThreshold(float InTriggerThreshold);
+	void SetTriggerThresholdVisibility(const bool InbIsVisible);
+
+	void UpdateSequenceRulerStyle(const FFixedSampleSequenceRulerStyle UpdatedRulerStyle);
+	void UpdateValueGridOverlayStyle(const FSampledSequenceValueGridOverlayStyle UpdatedValueGridOverlayStyle);
+	void UpdateSequenceViewerStyle(const FSampledSequenceViewerStyle UpdatedSequenceViewerStyle);
+	void UpdateTriggerThresholdStyle(const FTriggerThresholdLineStyle UpdatedStyle);
+
+	EAudioPanelLayoutType GetPanelLayoutType() { return PanelLayoutType; }
+
+	FOnSelectedChannelChanged  OnSelectedChannelChanged;
+	FOnTriggerModeChanged      OnTriggerModeChanged;
+	FOnTriggerThresholdChanged OnTriggerThresholdChanged;
+	FOnTimeWindowValueChanged  OnTimeWindowValueChanged;
+	FOnAnalysisPeriodChanged   OnAnalysisPeriodChanged;
+
+private:
+	void CreateLayout();
+
+	void CreateGridData(const FFixedSampleSequenceRulerStyle& RulerStyle);
+
+	// Basic panel methods
+	void CreateSequenceRuler(TSharedRef<FFixedSampledSequenceGridData> InGridData, const FFixedSampleSequenceRulerStyle& RulerStyle);
+	void CreateBackground(const FSampledSequenceViewerStyle& ViewerStyle);
+	void CreateValueGridOverlay(const uint32 MaxDivisionParameter, const SampledSequenceValueGridOverlay::EGridDivideMode DivideMode, const EYAxisLabelsUnit InValueGridOverlayDisplayUnit, const FSampledSequenceValueGridOverlayStyle& ValueGridStyle);
+	void CreateSequenceViewer(TSharedRef<FFixedSampledSequenceGridData> InGridData, const FFixedSampledSequenceView& InData, const FSampledSequenceViewerStyle& ViewerStyle);
+	void CreateTriggerThresholdLine(const FTriggerThresholdLineStyle& TriggerThresholdLineStyle);
+
+	// Advanced panel methods
+	void CreateChannelCombobox();
+	void CreateTriggerModeCombobox();
+	void CreateTriggerThresholdKnob();
+	void CreateTimeWindowKnob();
+	void CreateAnalysisPeriodKnob();
+	void CreateOscilloscopeControls();
+
+	const FAudioOscilloscopePanelStyle* PanelStyle = nullptr;
+
+	// Basic panel widgets
+	TSharedPtr<FFixedSampledSequenceGridData> SequenceGridData;
+	TSharedPtr<SBorder> BackgroundBorder;
+	TSharedPtr<SFixedSampledSequenceRuler> SequenceRuler;
+	TSharedPtr<SFixedSampledSequenceViewer> SequenceViewer;
+	TSharedPtr<SSampledSequenceValueGridOverlay> ValueGridOverlay;
+	TSharedPtr<STriggerThresholdLineWidget> TriggerThresholdLineWidget;
+
+	// Advanced panel widgets
+	TSharedPtr<SComboBox<TSharedPtr<FString>>> ChannelCombobox;
+	TArray<TSharedPtr<FString>> ChannelComboboxOptionsSource;
+	TSharedPtr<FString> SelectedChannelPtr;
+
+	TSharedPtr<SComboBox<TSharedPtr<EAudioOscilloscopeTriggerMode>>> TriggerModeCombobox;
+	TArray<TSharedPtr<EAudioOscilloscopeTriggerMode>> TriggerModeComboboxOptionsSource;
+	TSharedPtr<EAudioOscilloscopeTriggerMode> SelectedTriggerModePtr;
+
+	TSharedPtr<SAudioRadialSlider> TriggerThresholdKnob;
+	TSharedPtr<SAudioRadialSlider> TimeWindowKnob;
+	TSharedPtr<SAudioRadialSlider> AnalysisPeriodKnob;
+
+	EAudioPanelLayoutType PanelLayoutType = EAudioPanelLayoutType::Basic;
+
+	uint32 NumChannels = 0;
+
+	bool bIsInputWidgetTransacting = false;
+
+	float TriggerThresholdValue = 0.0f;
+	float TimeWindowValue       = 0.0f;
+	float AnalysisPeriodValue   = 0.0f;
+
+	float CachedPixelWidth           = 0.0f;
+	float OscilloscopeViewProportion = 1.0f;
+
+	EXAxisLabelsUnit SequenceRulerDisplayUnit;
+	EYAxisLabelsUnit ValueGridOverlayDisplayUnit;
+
+	FFixedSampledSequenceView DataView;
+	SampledSequenceDrawingUtils::FSampledSequenceDrawingParams DrawingParams;
+
+	bool bHideSequenceGrid  = false;
+	bool bHideSequenceRuler = false;
+	bool bHideValueGrid     = false;
+
+	bool bHideTriggerThresholdLine = true;
+};
diff --git a/Engine/Plugins/Runtime/AudioWidgets/Source/AudioWidgets/Public/SAudioVectorscopePanelWidget.h b/Engine/Plugins/Runtime/AudioWidgets/Source/AudioWidgets/Public/SAudioVectorscopePanelWidget.h
new file mode 100644
--- /dev/null
+++ b/Engine/Plugins/Runtime/AudioWidgets/Source/AudioWidgets/Public/SAudioVectorscopePanelWidget.h
@@ -0,0 +1,105 @@
+// Copyright Epic Games, Inc. All Rights Reserved.
+
+#pragma once
+
+#include "AudioVectorscopePanelStyle.h"
+#include "AudioWidgetsEnums.h"
+#include "Brushes/SlateColorBrush.h"
+#include "IFixedSampledSequenceViewReceiver.h"
+#include "SampledSequenceDrawingUtils.h"
+#include "SSampledSequenceValueGridOverlay.h"
+#include "Widgets/SCompoundWidget.h"
+
+class SAudioRadialSlider;
+class SBorder;
+class SFixedSampledSequenceViewer;
+class SFixedSampledSequenceVectorViewer;
+
+DECLARE_MULTICAST_DELEGATE_OneParam(FOnTimeWindowValueChanged, const float /*InTimeWindowMs*/)
+
+class AUDIOWIDGETS_API SAudioVectorscopePanelWidget : public SCompoundWidget, public IFixedSampledSequenceViewReceiver
+{
+public:
+	SLATE_BEGIN_ARGS(SAudioVectorscopePanelWidget)
+		: _HideGrid(false)
+		, _ValueGridMaxDivisionParameter(2)
+		, _PanelLayoutType(EAudioPanelLayoutType::Basic)
+		, _PanelStyle(&FAudioVectorscopePanelStyle::GetDefault())
+	{}
+		
+		SLATE_DEFAULT_SLOT(FArguments, InArgs)
+
+		/** Whether the value grid should be drawn or not */
+		SLATE_ATTRIBUTE(bool, HideGrid)
+
+		/** Maximum number of divisions in the value grid */
+		SLATE_ARGUMENT(uint32, ValueGridMaxDivisionParameter)
+
+		/** If we want to set the basic or advanced layout */
+		SLATE_ARGUMENT(EAudioPanelLayoutType, PanelLayoutType)
+
+		/** Vectorscope Panel widget style */
+		SLATE_STYLE_ARGUMENT(FAudioVectorscopePanelStyle, PanelStyle)
+
+	SLATE_END_ARGS()
+
+	void Construct(const FArguments& InArgs, const FFixedSampledSequenceView& InData);
+
+	// IFixedSampledSequenceViewReceiver Interface
+	virtual void ReceiveSequenceView(const FFixedSampledSequenceView InData, const uint32 FirstSampleIndex) override;
+
+	void SetGridVisibility(const bool InbIsVisible);
+	void SetValueGridOverlayMaxNumDivisions(const uint32 InGridMaxNumDivisions);
+
+	void SetVectorViewerScaleFactor(const float InScaleFactor);
+
+	void UpdateValueGridOverlayStyle(const FSampledSequenceValueGridOverlayStyle UpdatedValueGridOverlayStyle);
+	void UpdateSequenceVectorViewerStyle(const FSampledSequenceVectorViewerStyle UpdatedSequenceVectorViewerStyle);
+
+	EAudioPanelLayoutType GetPanelLayoutType() { return PanelLayoutType; }
+
+	FOnTimeWindowValueChanged OnTimeWindowValueChanged;
+
+private:
+	void CreateLayout();
+
+	// Basic panel methods
+	void CreateBackground(const FSampledSequenceVectorViewerStyle& VectorViewerStyle);
+
+	TSharedPtr<SSampledSequenceValueGridOverlay> CreateValueGridOverlay(const uint32 MaxDivisionParameter,
+		const SampledSequenceValueGridOverlay::EGridDivideMode DivideMode,
+		const FSampledSequenceValueGridOverlayStyle& ValueGridStyle,
+		const SampledSequenceDrawingUtils::ESampledSequenceDrawOrientation GridOrientation);
+
+	void CreateSequenceVectorViewer(const FFixedSampledSequenceView& InData, const FSampledSequenceVectorViewerStyle& VectorViewerStyle);
+
+	// Advanced panel methods
+	void CreateTimeWindowKnob();
+	void CreateScaleKnob();
+	void CreateVectorscopeControls();
+
+	const FAudioVectorscopePanelStyle* PanelStyle;
+
+	// Basic panel widgets
+	TSharedPtr<SBorder> BackgroundBorder;
+	TSharedPtr<SFixedSampledSequenceVectorViewer> SequenceVectorViewer;
+	TSharedPtr<SSampledSequenceValueGridOverlay> ValueGridOverlayXAxis;
+	TSharedPtr<SSampledSequenceValueGridOverlay> ValueGridOverlayYAxis;
+
+	// Advanced panel widgets
+	TSharedPtr<SAudioRadialSlider> TimeWindowKnob;
+	TSharedPtr<SAudioRadialSlider> ScaleKnob;
+
+	EAudioPanelLayoutType PanelLayoutType = EAudioPanelLayoutType::Basic;
+
+	bool bIsInputWidgetTransacting = false;
+
+	float TimeWindowValue = 0.0f;
+	float ScaleValue      = 0.0f;
+
+	float VectorscopeViewProportion = 1.0f;
+
+	FFixedSampledSequenceView DataView;
+
+	bool bHideValueGrid = false;
+};
diff --git a/Engine/Plugins/Runtime/AudioWidgets/Source/AudioWidgets/Public/SFixedSampledSequenceVectorViewer.h b/Engine/Plugins/Runtime/AudioWidgets/Source/AudioWidgets/Public/SFixedSampledSequenceVectorViewer.h
new file mode 100644
--- /dev/null
+++ b/Engine/Plugins/Runtime/AudioWidgets/Source/AudioWidgets/Public/SFixedSampledSequenceVectorViewer.h
@@ -0,0 +1,57 @@
+// Copyright Epic Games, Inc. All Rights Reserved.
+
+#pragma once
+
+#include "SampledSequenceDrawingUtils.h"
+#include "SampledSequenceVectorViewerStyle.h"
+#include "Widgets/SLeafWidget.h"
+
+
+class AUDIOWIDGETS_API SFixedSampledSequenceVectorViewer : public SLeafWidget
+{
+public: 
+	SLATE_BEGIN_ARGS(SFixedSampledSequenceVectorViewer)
+		: _ScaleFactor(1.0f)
+	{
+	}
+
+	SLATE_ARGUMENT(SampledSequenceDrawingUtils::FSampledSequenceDrawingParams, SequenceDrawingParams)
+
+	SLATE_ARGUMENT(float, ScaleFactor)
+
+	SLATE_STYLE_ARGUMENT(FSampledSequenceVectorViewerStyle, Style)
+
+	SLATE_END_ARGS()
+
+	void Construct(const FArguments& InArgs, TArrayView<const float> InSampleData, const uint8 InNumChannels);
+
+	void UpdateView(TArrayView<const float> InSampleData, const uint8 InNumChannels);
+	void SetScaleFactor(const float InScaleFactor);
+	void OnStyleUpdated(const FSampledSequenceVectorViewerStyle UpdatedStyle);
+
+private:
+	virtual int32 OnPaint(const FPaintArgs& Args, 
+		const FGeometry& AllottedGeometry, 
+		const FSlateRect& MyCullingRect, 
+		FSlateWindowElementList& OutDrawElements, 
+		int32 LayerId, 
+		const FWidgetStyle& InWidgetStyle, 
+		bool bParentEnabled) const override;
+
+	virtual FVector2D ComputeDesiredSize(float) const override;
+
+	bool bForceRedraw = false;
+
+	TArrayView<const float> SampleData;
+	SampledSequenceDrawingUtils::FSampledSequenceDrawingParams DrawingParams;
+
+	uint8 NumChannels = 0;
+
+	float DesiredHeight = 720.0f;
+	float DesiredWidth  = 720.0f;
+
+	FLinearColor LineColor = FLinearColor::White;
+	float LineThickness    = 1.0f;
+
+	float ScaleFactor = 1.0f;
+};
diff --git a/Engine/Plugins/Runtime/AudioWidgets/Source/AudioWidgets/Public/STriggerThresholdLineWidget.h b/Engine/Plugins/Runtime/AudioWidgets/Source/AudioWidgets/Public/STriggerThresholdLineWidget.h
new file mode 100644
--- /dev/null
+++ b/Engine/Plugins/Runtime/AudioWidgets/Source/AudioWidgets/Public/STriggerThresholdLineWidget.h
@@ -0,0 +1,49 @@
+// Copyright Epic Games, Inc. All Rights Reserved.
+
+#pragma once
+
+#include "Rendering/DrawElements.h"
+#include "SampledSequenceDrawingUtils.h"
+#include "Styling/ISlateStyle.h"
+#include "Styling/SlateWidgetStyleAsset.h"
+#include "TriggerThresholdLineStyle.h"
+#include "Widgets/SLeafWidget.h"
+
+struct FTriggerThresholdLineStyle;
+
+class AUDIOWIDGETS_API STriggerThresholdLineWidget : public SLeafWidget
+{
+public:
+    SLATE_BEGIN_ARGS(STriggerThresholdLineWidget)
+        : _TriggerThreshold(0.0f)
+		, _Style(&FTriggerThresholdLineStyle::GetDefault())
+    {}
+    
+		SLATE_ARGUMENT(float, TriggerThreshold)
+		SLATE_ARGUMENT(SampledSequenceDrawingUtils::FSampledSequenceDrawingParams, SequenceDrawingParams)
+		SLATE_STYLE_ARGUMENT(FTriggerThresholdLineStyle, Style)
+
+    SLATE_END_ARGS()
+
+    void Construct(const FArguments& InArgs);
+
+    virtual int32 OnPaint(const FPaintArgs& Args, 
+		const FGeometry& AllottedGeometry, 
+		const FSlateRect& MyCullingRect, 
+		FSlateWindowElementList& OutDrawElements, 
+		int32 LayerId, 
+		const FWidgetStyle& InWidgetStyle,
+		bool bParentEnabled) const override;
+
+	virtual FVector2D ComputeDesiredSize(float) const override;
+
+	void OnStyleUpdated(const FTriggerThresholdLineStyle UpdatedStyle);
+
+	void SetTriggerThreshold(float InTriggerThreshold);
+
+private:
+    float TriggerThreshold;
+	FLinearColor LineColor;
+	SampledSequenceDrawingUtils::FSampledSequenceDrawingParams DrawingParams;
+	const FTriggerThresholdLineStyle* Style = nullptr;
+};
diff --git a/Engine/Plugins/Runtime/AudioWidgets/Source/AudioWidgets/Public/SampledSequenceVectorViewerStyle.h b/Engine/Plugins/Runtime/AudioWidgets/Source/AudioWidgets/Public/SampledSequenceVectorViewerStyle.h
new file mode 100644
--- /dev/null
+++ b/Engine/Plugins/Runtime/AudioWidgets/Source/AudioWidgets/Public/SampledSequenceVectorViewerStyle.h
@@ -0,0 +1,45 @@
+// Copyright Epic Games, Inc. All Rights Reserved.
+
+#pragma once
+
+#include "Brushes/SlateColorBrush.h"
+#include "Styling/SlateBrush.h"
+#include "Styling/SlateWidgetStyle.h"
+
+#include "SampledSequenceVectorViewerStyle.generated.h"
+
+/**
+ * Represents the appearance of a trigger threshold line
+ */
+USTRUCT(BlueprintType)
+struct AUDIOWIDGETS_API FSampledSequenceVectorViewerStyle : public FSlateWidgetStyle
+{
+	GENERATED_USTRUCT_BODY()
+
+	FSampledSequenceVectorViewerStyle() {}
+
+	virtual const FName GetTypeName() const override { return TypeName; }
+	static const FSampledSequenceVectorViewerStyle& GetDefault() { static FSampledSequenceVectorViewerStyle Default; return Default; }
+	virtual void GetResources(TArray< const FSlateBrush* >& OutBrushes) const override { OutBrushes.Add(&BackgroundBrush); };
+
+	/** The background color. */
+	UPROPERTY(EditAnywhere, BlueprintReadWrite, Category = Appearance)
+	FSlateColor BackgroundColor = FLinearColor::Black;
+	FSampledSequenceVectorViewerStyle& SetBackgroundColor(const FSlateColor InBackgroundColor) { BackgroundColor = InBackgroundColor; return *this; }
+
+	UPROPERTY(EditAnywhere, BlueprintReadWrite, Category = Appearance)
+	FSlateBrush BackgroundBrush = FSlateColorBrush(FLinearColor::White);
+	FSampledSequenceVectorViewerStyle& SetBackgroundBrush(const FSlateBrush InBackgroundBrush) { BackgroundBrush = InBackgroundBrush; return *this; }
+
+	/** The vector view line color. */
+	UPROPERTY(EditAnywhere, BlueprintReadWrite, Category = Appearance)
+	FLinearColor LineColor = FLinearColor::White;
+	FSampledSequenceVectorViewerStyle& SetLineColor(const FLinearColor InLineColor) { LineColor = InLineColor; return *this; }
+
+	/** The vector view line thickness. */
+	UPROPERTY(EditAnywhere, BlueprintReadWrite, Category = Appearance)
+	float LineThickness = 1.0f;
+	FSampledSequenceVectorViewerStyle& SetLineThickness(const float InLineThickness) { LineThickness = InLineThickness; return *this; }
+
+	inline static FName TypeName = FName("FSampledSequenceVectorViewerStyle");
+};
diff --git a/Engine/Plugins/Runtime/AudioWidgets/Source/AudioWidgets/Public/TriggerThresholdLineStyle.h b/Engine/Plugins/Runtime/AudioWidgets/Source/AudioWidgets/Public/TriggerThresholdLineStyle.h
new file mode 100644
--- /dev/null
+++ b/Engine/Plugins/Runtime/AudioWidgets/Source/AudioWidgets/Public/TriggerThresholdLineStyle.h
@@ -0,0 +1,32 @@
+// Copyright Epic Games, Inc. All Rights Reserved.
+
+#pragma once
+
+#include "AudioWidgetsSlateTypes.h"
+#include "Fonts/SlateFontInfo.h"
+#include "Styling/SlateBrush.h"
+#include "Styling/SlateWidgetStyle.h"
+#include "Styling/StyleDefaults.h"
+
+#include "TriggerThresholdLineStyle.generated.h"
+
+/**
+ * Represents the appearance of a trigger threshold line
+ */
+USTRUCT(BlueprintType)
+struct AUDIOWIDGETS_API FTriggerThresholdLineStyle : public FSlateWidgetStyle
+{
+	GENERATED_USTRUCT_BODY()
+
+	FTriggerThresholdLineStyle() {}
+
+	virtual const FName GetTypeName() const override { return TypeName; }
+	static const FTriggerThresholdLineStyle& GetDefault() { static FTriggerThresholdLineStyle Default; return Default; }
+
+	/** The trigger threshold line color. */
+	UPROPERTY(EditAnywhere, BlueprintReadWrite, Category = Appearance)
+	FLinearColor LineColor = FLinearColor::White;
+	FTriggerThresholdLineStyle& SetLineColor(const FLinearColor InLineColor) { LineColor = InLineColor; return *this; }
+
+	inline static FName TypeName = FName("FTriggerThresholdLineStyle");
+};
diff --git a/Engine/Plugins/Runtime/AudioWidgets/Source/AudioWidgets/Public/WaveformAudioSamplesDataProvider.h b/Engine/Plugins/Runtime/AudioWidgets/Source/AudioWidgets/Public/WaveformAudioSamplesDataProvider.h
new file mode 100644
--- /dev/null
+++ b/Engine/Plugins/Runtime/AudioWidgets/Source/AudioWidgets/Public/WaveformAudioSamplesDataProvider.h
@@ -0,0 +1,90 @@
+// Copyright Epic Games, Inc. All Rights Reserved.
+
+#pragma once
+
+#include "AudioOscilloscopeEnums.h"
+#include "Containers/ArrayView.h"
+#include "Containers/Ticker.h"
+#include "DSP/Dsp.h"
+#include "DSP/MultithreadedPatching.h"
+#include "Engine/World.h"
+#include "FixedSampledSequenceView.h"
+#include "IFixedSampledSequenceViewProvider.h"
+
+class UAudioBus;
+
+struct FWaveformAudioSamplesResult;
+
+DECLARE_MULTICAST_DELEGATE_TwoParams(FOnDataViewGenerated, FFixedSampledSequenceView /*InView*/, const uint32 /*FirstSampleIndex*/)
+
+namespace AudioWidgets
+{
+	class AUDIOWIDGETS_API FWaveformAudioSamplesDataProvider : public IFixedSampledSequenceViewProvider, 
+															   public TSharedFromThis<FWaveformAudioSamplesDataProvider>
+	{
+	public:
+		FWaveformAudioSamplesDataProvider(UWorld* InWorld, 
+			UAudioBus* InAudioBus, 
+			const uint32 InNumChannelToProvide, 
+			const float InTimeWindowMs, 
+			const float InMaxTimeWindowMs, 
+			const float InAnalysisPeriodMs);
+
+		virtual ~FWaveformAudioSamplesDataProvider();
+
+		void ResetAudioBuffers();
+
+		void StartProcessing();
+		void StopProcessing();
+
+		FFixedSampledSequenceView GetDataView() { return DataView; };
+		uint32 GetNumChannels() { return NumChannels; }
+		const UAudioBus* GetAudioBus() { return AudioBus; }
+
+		void SetChannelToAnalyze(const int32 InChannel);
+		void SetTriggerMode(const EAudioOscilloscopeTriggerMode InTriggerMode);
+		void SetTriggerThreshold(const float InTriggerThreshold);
+		void SetTimeWindow(const float InTimeWindowMs);
+		void SetAnalysisPeriod(const float InAnalysisPeriodMs);
+
+		virtual FFixedSampledSequenceView RequestSequenceView(const TRange<double> DataRatioRange) override;
+
+		FOnDataViewGenerated OnDataViewGenerated;
+
+	private:
+		void PushAudioSamplesToCircularBuffer();
+		bool Tick(float DeltaTime);
+
+		UWorld* WorldPtr = nullptr;
+		FTSTicker::FDelegateHandle TickerHandle = nullptr;
+
+		uint32 NumChannelsToProvide = 0;
+		uint32 NumChannels          = 0;
+		uint32 SampleRate           = 0;
+
+		uint32 TimeWindowMaxTimeSamples = 0;
+
+		const Audio::FMixerDevice* MixerDevice = nullptr;
+
+		UAudioBus* AudioBus = nullptr;
+		Audio::FPatchOutputStrongPtr PatchOutput = nullptr;
+
+		TArray<float> TempAudioBuffer;
+		Audio::TCircularAudioBuffer<float> AudioSamplesCircularBuffer;
+		TArray<float> AudioSamplesForView;
+		FFixedSampledSequenceView DataView;
+
+		int32 ChannelIndexToAnalyze = 0;
+
+		EAudioOscilloscopeTriggerMode TriggerMode = EAudioOscilloscopeTriggerMode::None;
+		float TriggerThreshold = 0.0f;
+
+		uint32 TimeWindowSamples     = 0;
+		uint32 AnalysisPeriodSamples = 0;
+
+		bool bIsProcessing = false;
+		bool bHasTriggered = false;
+
+		uint32 NumSamplesPushedToCircularBuffer = 0;
+	};
+} // namespace AudioWidgets
diff --git a/Engine/Source/Runtime/Engine/Classes/Sound/AudioBus.h b/Engine/Source/Runtime/Engine/Classes/Sound/AudioBus.h
--- a/Engine/Source/Runtime/Engine/Classes/Sound/AudioBus.h
+++ b/Engine/Source/Runtime/Engine/Classes/Sound/AudioBus.h
@@ -19,6 +19,26 @@ enum class EAudioBusChannels : uint8
 	SevenPointOne = 7
 };
 
+namespace AudioBusUtils
+{
+	static EAudioBusChannels ConvertIntToEAudioBusChannels(const int32 InValue)
+	{
+		switch (InValue)
+		{
+		case 1: return EAudioBusChannels::Mono;
+		case 2:	return EAudioBusChannels::Stereo;
+		case 4:	return EAudioBusChannels::Quad;
+		case 6:	return EAudioBusChannels::FivePointOne;
+		case 8:	return EAudioBusChannels::SevenPointOne;
+		default:
+			// TODO alex.perez: LogAudio undeclared identifier, is it ok to add the header that has LogAudio, or should we move this function to another place?
+			//UE_LOG(LogAudio, Error, TEXT("Number of channels: %d not available in AudioBusChannels configurations. Make sure that the number of channels used is 1, 2, 4, 6 or 8."), InValue);
+			check(false);
+			return EAudioBusChannels::Mono;
+		}
+	};
+}
+
 class FAudioBusProxy;
 
 using FAudioBusProxyPtr = TSharedPtr<FAudioBusProxy, ESPMode::ThreadSafe>;
