[AudioWidgets] Added missing headers in SFixedSampledSequenceVectorViewer.h

diff --git a/Engine/Plugins/Runtime/AudioWidgets/Source/AudioWidgets/Public/SFixedSampledSequenceVectorViewer.h b/Engine/Plugins/Runtime/AudioWidgets/Source/AudioWidgets/Public/SFixedSampledSequenceVectorViewer.h
--- a/Engine/Plugins/Runtime/AudioWidgets/Source/AudioWidgets/Public/SFixedSampledSequenceVectorViewer.h
+++ b/Engine/Plugins/Runtime/AudioWidgets/Source/AudioWidgets/Public/SFixedSampledSequenceVectorViewer.h
@@ -4,6 +4,8 @@
 
 #include "SampledSequenceDrawingUtils.h"
 #include "SampledSequenceVectorViewerStyle.h"
+#include "Styling/ISlateStyle.h"
+#include "Styling/SlateWidgetStyleAsset.h"
 #include "Widgets/SLeafWidget.h"
 
 
