[AudioWidgets] Adjusted copyright comment in AudioWidgetsLog.h

diff --git a/Engine/Plugins/Runtime/AudioWidgets/Source/AudioWidgets/Public/AudioWidgetsLog.h b/Engine/Plugins/Runtime/AudioWidgets/Source/AudioWidgets/Public/AudioWidgetsLog.h
--- a/Engine/Plugins/Runtime/AudioWidgets/Source/AudioWidgets/Public/AudioWidgetsLog.h
+++ b/Engine/Plugins/Runtime/AudioWidgets/Source/AudioWidgets/Public/AudioWidgetsLog.h
@@ -1,4 +1,4 @@
-// Copyright 2023, SweejTech Ltd. All Rights Reserved.
+// Copyright Epic Games, Inc. All Rights Reserved.
 
 #pragma once
 
