[AudioWidgets] Oscilloscope and Vectorscope panel widgets: Build widgets considering the panel mode (basic or advanced) + small bug fix in oscilloscope widget (missing combobox item source Empty())

diff --git a/Engine/Plugins/Runtime/AudioWidgets/Source/AudioWidgets/Private/AudioOscilloscope.cpp b/Engine/Plugins/Runtime/AudioWidgets/Source/AudioWidgets/Private/AudioOscilloscope.cpp
--- a/Engine/Plugins/Runtime/AudioWidgets/Source/AudioWidgets/Private/AudioOscilloscope.cpp
+++ b/Engine/Plugins/Runtime/AudioWidgets/Source/AudioWidgets/Private/AudioOscilloscope.cpp
@@ -51,7 +51,7 @@ namespace AudioWidgets
 		}
 		else
 		{
-			OscilloscopePanelWidget->BuildWidget(SequenceView, InNumChannels);
+			OscilloscopePanelWidget->BuildWidget(SequenceView, InNumChannels, InPanelLayoutType);
 		}
 
 		// Interconnect data provider and widget
diff --git a/Engine/Plugins/Runtime/AudioWidgets/Source/AudioWidgets/Private/AudioVectorscope.cpp b/Engine/Plugins/Runtime/AudioWidgets/Source/AudioWidgets/Private/AudioVectorscope.cpp
--- a/Engine/Plugins/Runtime/AudioWidgets/Source/AudioWidgets/Private/AudioVectorscope.cpp
+++ b/Engine/Plugins/Runtime/AudioWidgets/Source/AudioWidgets/Private/AudioVectorscope.cpp
@@ -47,7 +47,7 @@ namespace AudioWidgets
 		}
 		else
 		{
-			VectorscopePanelWidget->BuildWidget(SequenceView);
+			VectorscopePanelWidget->BuildWidget(SequenceView, InPanelLayoutType);
 		}
 
 		// Interconnect data provider and widget
diff --git a/Engine/Plugins/Runtime/AudioWidgets/Source/AudioWidgets/Private/SAudioOscilloscopePanelWidget.cpp b/Engine/Plugins/Runtime/AudioWidgets/Source/AudioWidgets/Private/SAudioOscilloscopePanelWidget.cpp
--- a/Engine/Plugins/Runtime/AudioWidgets/Source/AudioWidgets/Private/SAudioOscilloscopePanelWidget.cpp
+++ b/Engine/Plugins/Runtime/AudioWidgets/Source/AudioWidgets/Private/SAudioOscilloscopePanelWidget.cpp
@@ -71,8 +71,6 @@ namespace AudioOscilloscopePanelWidgetPrivate
 
 void SAudioOscilloscopePanelWidget::Construct(const FArguments& InArgs, const FFixedSampledSequenceView& InData, const int32 InNumChannels)
 {
-	PanelLayoutType = InArgs._PanelLayoutType;
-
 	SequenceRulerDisplayUnit      = InArgs._SequenceRulerDisplayUnit;
 	ValueGridOverlayDisplayUnit   = InArgs._YAxisLabelsUnit.Get();
 	ValueGridMaxDivisionParameter = InArgs._ValueGridMaxDivisionParameter;
@@ -84,12 +82,13 @@ void SAudioOscilloscopePanelWidget::Construct(const FArguments& InArgs, const FF
 	check(InArgs._PanelStyle);
 	PanelStyle = InArgs._PanelStyle;
 
-	BuildWidget(InData, InNumChannels);
+	BuildWidget(InData, InNumChannels, InArgs._PanelLayoutType);
 }
 
-void SAudioOscilloscopePanelWidget::BuildWidget(const FFixedSampledSequenceView& InData, const int32 InNumChannels)
+void SAudioOscilloscopePanelWidget::BuildWidget(const FFixedSampledSequenceView& InData, const int32 InNumChannels, const EAudioPanelLayoutType InPanelLayoutType)
 {
-	NumChannels = InNumChannels;
+	NumChannels     = InNumChannels;
+	PanelLayoutType = InPanelLayoutType;
 
 	DataView = InData;
 
@@ -385,6 +384,7 @@ void SAudioOscilloscopePanelWidget::CreateChannelCombobox()
 		return SelectedChannelPtr.IsValid() ? FText::FromString(*SelectedChannelPtr) : FText::GetEmpty();
 	};
 
+	ChannelComboboxOptionsSource.Empty();
 	for (uint32 Index = 0; Index < NumChannels; ++Index)
 	{
 		ChannelComboboxOptionsSource.Add(MakeShareable(new FString(FString::FromInt(Index + 1))));
@@ -429,6 +429,7 @@ void SAudioOscilloscopePanelWidget::CreateTriggerModeCombobox()
 		return SelectedTriggerModePtr.IsValid() ? FText::FromString(*SelectedTriggerModeEnumStr) : FText::GetEmpty();
 	};
 
+	TriggerModeComboboxOptionsSource.Empty();
 	for (int64 Index = 0; Index < StaticEnum<EAudioOscilloscopeTriggerMode>()->GetMaxEnumValue(); ++Index)
 	{
 		TriggerModeComboboxOptionsSource.Add(MakeShareable(new EAudioOscilloscopeTriggerMode(static_cast<EAudioOscilloscopeTriggerMode>(Index))));
diff --git a/Engine/Plugins/Runtime/AudioWidgets/Source/AudioWidgets/Private/SAudioVectorscopePanelWidget.cpp b/Engine/Plugins/Runtime/AudioWidgets/Source/AudioWidgets/Private/SAudioVectorscopePanelWidget.cpp
--- a/Engine/Plugins/Runtime/AudioWidgets/Source/AudioWidgets/Private/SAudioVectorscopePanelWidget.cpp
+++ b/Engine/Plugins/Runtime/AudioWidgets/Source/AudioWidgets/Private/SAudioVectorscopePanelWidget.cpp
@@ -20,20 +20,19 @@
 
 void SAudioVectorscopePanelWidget::Construct(const FArguments& InArgs, const FFixedSampledSequenceView& InData)
 {
-	PanelLayoutType = InArgs._PanelLayoutType;
-
 	ValueGridMaxDivisionParameter = InArgs._ValueGridMaxDivisionParameter;
 	bHideValueGrid = InArgs._HideGrid.Get();
 
 	check(InArgs._PanelStyle);
 	PanelStyle = InArgs._PanelStyle;
 
-	BuildWidget(InData);
+	BuildWidget(InData, InArgs._PanelLayoutType);
 }
 
-void SAudioVectorscopePanelWidget::BuildWidget(const FFixedSampledSequenceView& InData)
+void SAudioVectorscopePanelWidget::BuildWidget(const FFixedSampledSequenceView& InData, const EAudioPanelLayoutType InPanelLayoutType)
 {
-	DataView = InData;
+	DataView        = InData;
+	PanelLayoutType = InPanelLayoutType;
 
 	CreateBackground(PanelStyle->VectorViewerStyle);
 
diff --git a/Engine/Plugins/Runtime/AudioWidgets/Source/AudioWidgets/Public/SAudioOscilloscopePanelWidget.h b/Engine/Plugins/Runtime/AudioWidgets/Source/AudioWidgets/Public/SAudioOscilloscopePanelWidget.h
--- a/Engine/Plugins/Runtime/AudioWidgets/Source/AudioWidgets/Public/SAudioOscilloscopePanelWidget.h
+++ b/Engine/Plugins/Runtime/AudioWidgets/Source/AudioWidgets/Public/SAudioOscilloscopePanelWidget.h
@@ -91,7 +91,7 @@ public:
 
 	void Construct(const FArguments& InArgs, const FFixedSampledSequenceView& InData, const int32 InNumChannels);
 
-	void BuildWidget(const FFixedSampledSequenceView& InData, const int32 InNumChannels);
+	void BuildWidget(const FFixedSampledSequenceView& InData, const int32 InNumChannels, const EAudioPanelLayoutType InPanelLayoutType);
 
 	// SWidget methods overrides
 	virtual void Tick(const FGeometry& AllottedGeometry, const double InCurrentTime, const float InDeltaTime) override;
diff --git a/Engine/Plugins/Runtime/AudioWidgets/Source/AudioWidgets/Public/SAudioVectorscopePanelWidget.h b/Engine/Plugins/Runtime/AudioWidgets/Source/AudioWidgets/Public/SAudioVectorscopePanelWidget.h
--- a/Engine/Plugins/Runtime/AudioWidgets/Source/AudioWidgets/Public/SAudioVectorscopePanelWidget.h
+++ b/Engine/Plugins/Runtime/AudioWidgets/Source/AudioWidgets/Public/SAudioVectorscopePanelWidget.h
@@ -45,7 +45,7 @@ public:
 
 	void Construct(const FArguments& InArgs, const FFixedSampledSequenceView& InData);
 
-	void BuildWidget(const FFixedSampledSequenceView& InData);
+	void BuildWidget(const FFixedSampledSequenceView& InData, const EAudioPanelLayoutType InPanelLayoutType);
 
 	// IFixedSampledSequenceViewReceiver Interface
 	virtual void ReceiveSequenceView(const FFixedSampledSequenceView InData, const uint32 FirstSampleIndex) override;
