[AudioWidgets] Bug fix for: [MetaSounds][Widgets] - Playing a metasound twice generates errors in the output log

#rnx

diff --git a/Engine/Plugins/Runtime/AudioWidgets/Source/AudioWidgets/Private/WaveformAudioSamplesDataProvider.cpp b/Engine/Plugins/Runtime/AudioWidgets/Source/AudioWidgets/Private/WaveformAudioSamplesDataProvider.cpp
--- a/Engine/Plugins/Runtime/AudioWidgets/Source/AudioWidgets/Private/WaveformAudioSamplesDataProvider.cpp
+++ b/Engine/Plugins/Runtime/AudioWidgets/Source/AudioWidgets/Private/WaveformAudioSamplesDataProvider.cpp
@@ -112,7 +112,6 @@ namespace AudioWidgets
 	{
 		if (bIsProcessing)
 		{
-			UE_LOG(LogAudioWidgets, Error, TEXT("StartProcessing invoked while having an active processing. Try call StopProcessing before calling StartProcessing."));
 			return;
 		}
 
