Reduce unnecessary memory allocs in Audio UMG widget on client and server.

diff --git a/Engine/Plugins/Runtime/AudioWidgets/Source/AudioWidgets/Private/AudioOscilloscopeUMG.cpp b/Engine/Plugins/Runtime/AudioWidgets/Source/AudioWidgets/Private/AudioOscilloscopeUMG.cpp
--- a/Engine/Plugins/Runtime/AudioWidgets/Source/AudioWidgets/Private/AudioOscilloscopeUMG.cpp
+++ b/Engine/Plugins/Runtime/AudioWidgets/Source/AudioWidgets/Private/AudioOscilloscopeUMG.cpp
@@ -20,8 +20,13 @@ UAudioOscilloscope::UAudioOscilloscope(const FObjectInitializer& ObjectInitializ
 	bCanChildrenBeAccessible = false;
 #endif
 
-	DummyAudioSamples.Init(0.0f, DummyMaxNumSamples);
-	DummyDataView = FFixedSampledSequenceView{ MakeArrayView(DummyAudioSamples.GetData(), DummyAudioSamples.Num()), DummyNumChannels, DummySampleRate };
+#if !UE_SERVER
+	if (!HasAnyFlags(RF_ClassDefaultObject))
+	{
+		DummyAudioSamples.Init(0.0f, DummyMaxNumSamples);
+		DummyDataView = FFixedSampledSequenceView{ MakeArrayView(DummyAudioSamples.GetData(), DummyAudioSamples.Num()), DummyNumChannels, DummySampleRate };
+	}
+#endif
 }
 
 void UAudioOscilloscope::CreateDummyOscilloscopeWidget()
