[AudioWidgets] Make AudioMeter a bit more RAII (from [REDACTED_HANDLE] with small adjustments) + added the possibility to construct the audio meter using an external audio bus (needed for AudioInsights).

#rnx

diff --git a/Engine/Plugins/Runtime/AudioWidgets/Source/AudioWidgets/Private/AudioMeter.cpp b/Engine/Plugins/Runtime/AudioWidgets/Source/AudioWidgets/Private/AudioMeter.cpp
--- a/Engine/Plugins/Runtime/AudioWidgets/Source/AudioWidgets/Private/AudioMeter.cpp
+++ b/Engine/Plugins/Runtime/AudioWidgets/Source/AudioWidgets/Private/AudioMeter.cpp
@@ -1,6 +1,9 @@
 // Copyright Epic Games, Inc. All Rights Reserved.
 
 #include "AudioMeter.h"
+
+#include "AudioBusSubsystem.h"
+#include "AudioMixerDevice.h"
 #include "SAudioMeter.h"
 #include "UObject/UObjectGlobals.h"
 
@@ -158,7 +161,7 @@ const FText UAudioMeter::GetPaletteCategory()
 
 namespace AudioWidgets
 {
-	FAudioMeter::FAudioMeter()
+	FAudioMeter::FAudioMeter(int32 InNumChannels, UWorld& InWorld, TObjectPtr<UAudioBus> InExternalAudioBus)
 		: Widget(SNew(SAudioMeter)
 			.Orientation(EOrientation::Orient_Vertical)
 			.BackgroundColor(FLinearColor::Transparent)
@@ -172,6 +175,12 @@ namespace AudioWidgets
 			.MeterScaleLabelColor(FLinearColor(0.442708f, 0.442708f, 0.442708f, 1.0f))
 		)
 	{
+		Init(InNumChannels, InWorld, InExternalAudioBus);
+	}
+
+	FAudioMeter::~FAudioMeter()
+	{
+		Teardown();
 	}
 
 	UAudioBus* FAudioMeter::GetAudioBus() const
@@ -184,7 +193,7 @@ namespace AudioWidgets
 		return StaticCastSharedRef<SAudioMeter>(Widget->AsShared());
 	}
 
-	void FAudioMeter::Init(int32 InNumChannels, UWorld& InWorld)
+	void FAudioMeter::Init(int32 InNumChannels, UWorld& InWorld, TObjectPtr<UAudioBus> InExternalAudioBus)
 	{
 		check(InNumChannels > 0);
 
@@ -196,7 +205,9 @@ namespace AudioWidgets
 		Analyzer = TStrongObjectPtr(NewObject<UMeterAnalyzer>());
 		Analyzer->Settings = Settings.Get();
 
-		AudioBus = TStrongObjectPtr(NewObject<UAudioBus>());
+		bUseExternalAudioBus = InExternalAudioBus != nullptr;
+
+		AudioBus = bUseExternalAudioBus ? TStrongObjectPtr(InExternalAudioBus.Get()) : TStrongObjectPtr(NewObject<UAudioBus>());
 		AudioBus->AudioBusChannels = EAudioBusChannels(InNumChannels - 1);
 
 		ResultsDelegateHandle = Analyzer->OnLatestPerChannelMeterResultsNative.AddRaw(this, &FAudioMeter::OnMeterOutput);
@@ -240,17 +251,23 @@ namespace AudioWidgets
 
 	void FAudioMeter::Teardown()
 	{
-		if (Analyzer.IsValid() && ResultsDelegateHandle.IsValid())
+		if (Analyzer.IsValid() && Analyzer->IsValidLowLevel())
 		{
 			Analyzer->StopAnalyzing();
-			Analyzer->OnLatestPerChannelMeterResultsNative.Remove(ResultsDelegateHandle);
+			if (ResultsDelegateHandle.IsValid())
+			{
+				Analyzer->OnLatestPerChannelMeterResultsNative.Remove(ResultsDelegateHandle);
+			}
+
 			Analyzer.Reset();
-			ResultsDelegateHandle.Reset();
 		}
 
+		ResultsDelegateHandle.Reset();
 		AudioBus.Reset();
 		ChannelInfo.Reset();
 		Settings.Reset();
+
+		bUseExternalAudioBus = false;
 	}
 } // namespace AudioWidgets
 #undef LOCTEXT_NAMESPACE
diff --git a/Engine/Plugins/Runtime/AudioWidgets/Source/AudioWidgets/Public/AudioMeter.h b/Engine/Plugins/Runtime/AudioWidgets/Source/AudioWidgets/Public/AudioMeter.h
--- a/Engine/Plugins/Runtime/AudioWidgets/Source/AudioWidgets/Public/AudioMeter.h
+++ b/Engine/Plugins/Runtime/AudioWidgets/Source/AudioWidgets/Public/AudioMeter.h
@@ -155,20 +155,21 @@ namespace AudioWidgets
 	class AUDIOWIDGETS_API FAudioMeter
 	{
 	public:
-		FAudioMeter();
+		FAudioMeter(int32 InNumChannels, UWorld& InWorld, TObjectPtr<UAudioBus> InExternalAudioBus = nullptr); // OPTIONAL PARAM InExternalAudioBus: An audio meter can be constructed from this audio bus.
+		~FAudioMeter();
 
 		UAudioBus* GetAudioBus() const;
 
 		TSharedRef<SAudioMeter> GetWidget() const;
 
-		void Teardown();
-
-		void Init(int32 InNumChannels, UWorld& InWorld);
+		void Init(int32 InNumChannels, UWorld& InWorld, TObjectPtr<UAudioBus> InExternalAudioBus = nullptr);
 
 	protected:
 		void OnMeterOutput(UMeterAnalyzer* InMeterAnalyzer, int32 ChannelIndex, const FMeterResults& InMeterResults);
 
 	private:
+		void Teardown();
+
 		/** Metasound analyzer object. */
 		TStrongObjectPtr<UMeterAnalyzer> Analyzer;
 
@@ -188,5 +189,7 @@ namespace AudioWidgets
 		TSharedPtr<SAudioMeter> Widget;
 
 		TWeakObjectPtr<UWorld> WorldPtr;
+
+		bool bUseExternalAudioBus = false;
 	};
 } // namespace AudioWidgets
diff --git a/Engine/Plugins/Runtime/Metasound/Source/MetasoundEditor/Private/MetasoundEditor.cpp b/Engine/Plugins/Runtime/Metasound/Source/MetasoundEditor/Private/MetasoundEditor.cpp
--- a/Engine/Plugins/Runtime/Metasound/Source/MetasoundEditor/Private/MetasoundEditor.cpp
+++ b/Engine/Plugins/Runtime/Metasound/Source/MetasoundEditor/Private/MetasoundEditor.cpp
@@ -1162,17 +1162,15 @@ namespace Metasound
 		{
 			if (UMetaSoundSource* MetaSoundSource = Cast<UMetaSoundSource>(Metasound))
 			{
-				if (!OutputMeter.IsValid())
-				{
-					OutputMeter = MakeShared<AudioWidgets::FAudioMeter>();
-				}
-
 				if (ensure(GEditor))
 				{
 					UWorld* EditorWorld = GEditor->GetEditorWorldContext().World();
 					if (ensure(EditorWorld))
 					{
-						OutputMeter->Init(MetaSoundSource->NumChannels, *EditorWorld);
+						if (!OutputMeter.IsValid())
+						{
+							OutputMeter = MakeShared<AudioWidgets::FAudioMeter>(MetaSoundSource->NumChannels, *EditorWorld);
+						}
 
 						const uint32 MetaSoundNumChannels = static_cast<uint32>(MetaSoundSource->NumChannels);
 
@@ -1220,33 +1218,19 @@ namespace Metasound
 							OutputVectorscope->CreateVectorscopeWidget(VectorscopePanelLayoutType);
 						}
 
+						return;
 					}
 				}
 			}
-			else
-			{
-				OutputMeter.Reset();
-				OutputOscilloscope.Reset();
-				OutputVectorscope.Reset();
-			}
+
+			DestroyAnalyzers();
 		}
 
 		void FEditor::DestroyAnalyzers()
 		{
-			if (OutputMeter.IsValid())
-			{
-				OutputMeter->Teardown();
-			}
-
-			if (OutputOscilloscope.IsValid())
-			{
-				OutputOscilloscope->StopProcessing();
-			}
-
-			if (OutputVectorscope.IsValid())
-			{
-				OutputVectorscope->StopProcessing();
-			}
+			OutputMeter.Reset();
+			OutputOscilloscope.Reset();
+			OutputVectorscope.Reset();
 		}
 
 		void FEditor::ExtendToolbar()
