[AudioWidgets] Added SAudioSpectrumPlot Slate Widget and Style. Used for plotting an audio power spectrum, with linear or log frequency scale, and decibels sound levels.
Took some inspiration from SCurveEditor for the implementation.

diff --git a/Engine/Plugins/Runtime/AudioWidgets/Source/AudioWidgets/Private/AudioSpectrumPlotStyle.cpp b/Engine/Plugins/Runtime/AudioWidgets/Source/AudioWidgets/Private/AudioSpectrumPlotStyle.cpp
new file mode 100644
--- /dev/null
+++ b/Engine/Plugins/Runtime/AudioWidgets/Source/AudioWidgets/Private/AudioSpectrumPlotStyle.cpp
@@ -0,0 +1,20 @@
+// Copyright Epic Games, Inc. All Rights Reserved.
+
+#include "AudioSpectrumPlotStyle.h"
+#include "Styling/StyleDefaults.h"
+
+#include UE_INLINE_GENERATED_CPP_BY_NAME(AudioSpectrumPlotStyle)
+
+FAudioSpectrumPlotStyle::FAudioSpectrumPlotStyle()
+{
+}
+
+
+const FName FAudioSpectrumPlotStyle::TypeName("FAudioSpectrumPlotStyle");
+
+const FAudioSpectrumPlotStyle& FAudioSpectrumPlotStyle::GetDefault()
+{
+	static FAudioSpectrumPlotStyle Default;
+	return Default;
+}
+
diff --git a/Engine/Plugins/Runtime/AudioWidgets/Source/AudioWidgets/Private/SAudioSpectrumPlot.cpp b/Engine/Plugins/Runtime/AudioWidgets/Source/AudioWidgets/Private/SAudioSpectrumPlot.cpp
new file mode 100644
--- /dev/null
+++ b/Engine/Plugins/Runtime/AudioWidgets/Source/AudioWidgets/Private/SAudioSpectrumPlot.cpp
@@ -0,0 +1,518 @@
+// Copyright Epic Games, Inc. All Rights Reserved.
+
+#include "SAudioSpectrumPlot.h"
+
+#include "Fonts/FontMeasure.h"
+#include "Framework/Application/SlateApplication.h"
+#include "SlateOptMacros.h"
+
+BEGIN_SLATE_FUNCTION_BUILD_OPTIMIZATION
+
+/**
+ * Helper class for drawing grid lines with text labels. Includes logic to avoid drawing overlapping labels if the grid lines are close together.
+ */
+class FAudioSpectrumPlotGridAndLabelDrawingHelper
+{
+public:
+	FAudioSpectrumPlotGridAndLabelDrawingHelper(const FGeometry& InAllottedGeometry, FSlateWindowElementList& OutDrawElements, const FAudioSpectrumPlotScaleInfo& InScaleInfo);
+
+	void DrawSoundLevelGridLines(const int32 LayerId, TConstArrayView<float> GridLineSoundLevels, const FLinearColor& LineColor) const;
+	void DrawFrequencyGridLines(const int32 LayerId, TConstArrayView<float> GridLineFrequencies, const FLinearColor& LineColor) const;
+
+	void DrawSoundLevelAxisLabels(const int32 LayerId, TConstArrayView<float> GridLineSoundLevels, const FSlateFontInfo& Font, const FLinearColor& TextColor);
+	void DrawFrequencyAxisLabels(const int32 LayerId, TConstArrayView<float> GridLineFrequencies, const FSlateFontInfo& Font, const FLinearColor& TextColor);
+
+	bool HasDrawnLabels() const { return !DrawnLabelRects.IsEmpty(); }
+
+private:
+	static FString FormatSoundLevelString(const float SoundLevel);
+	static FString FormatFreqString(const float Freq);
+
+	void DrawLabelIfNoOverlap(const int32 LayerId, const float LabelLeft, const float LabelTop, const FVector2f& LabelDrawSize, const FString LabelText, const FSlateFontInfo& Font, const FLinearColor& TextColor);
+
+	// Tweak the label Rect bounds to give space where it's needed for readability, while not wasting space where it's not needed.
+	FSlateRect GetModifiedLabelRect(const FSlateRect& LabelRect) const;
+
+	bool IsOverlappingPreviouslyDrawnLabel(const FSlateRect& LabelRect) const;
+
+	const FGeometry& AllottedGeometry;
+	FSlateWindowElementList& ElementList;
+	const FAudioSpectrumPlotScaleInfo& ScaleInfo;
+	const FSlateRect LocalBackgroundRect;
+	const TSharedRef<FSlateFontMeasure> FontMeasureService;
+	FVector2f SpaceDrawSize; // Cached draw size of a space character.
+	TArray<FSlateRect> DrawnLabelRects; // Keep track of where text labels have been drawn.
+};
+
+FAudioSpectrumPlotGridAndLabelDrawingHelper::FAudioSpectrumPlotGridAndLabelDrawingHelper(const FGeometry& InAllottedGeometry, FSlateWindowElementList& OutDrawElements, const FAudioSpectrumPlotScaleInfo& InScaleInfo)
+	: AllottedGeometry(InAllottedGeometry)
+	, ElementList(OutDrawElements)
+	, ScaleInfo(InScaleInfo)
+	, LocalBackgroundRect(FVector2f::ZeroVector, InAllottedGeometry.GetLocalSize())
+	, FontMeasureService(FSlateApplication::Get().GetRenderer()->GetFontMeasureService())
+	, SpaceDrawSize(FVector2f::ZeroVector)
+{
+	//
+}
+
+void FAudioSpectrumPlotGridAndLabelDrawingHelper::DrawSoundLevelGridLines(const int32 LayerId, TConstArrayView<float> GridLineSoundLevels, const FLinearColor& LineColor) const
+{
+	TArray<FVector2f> LinePoints;
+	LinePoints.SetNum(2);
+
+	for (float SoundLevel : GridLineSoundLevels)
+	{
+		// Draw horizontal grid line:
+		const float GridLineLocalY = ScaleInfo.SoundLevelToLocalY(SoundLevel);
+		LinePoints[0] = { LocalBackgroundRect.Left, GridLineLocalY };
+		LinePoints[1] = { LocalBackgroundRect.Right, GridLineLocalY };
+		FSlateDrawElement::MakeLines(ElementList, LayerId, AllottedGeometry.ToPaintGeometry(), LinePoints, ESlateDrawEffect::None, LineColor);
+	}
+}
+
+void FAudioSpectrumPlotGridAndLabelDrawingHelper::DrawFrequencyGridLines(const int32 LayerId, TConstArrayView<float> GridLineFrequencies, const FLinearColor& LineColor) const
+{
+	TArray<FVector2f> LinePoints;
+	LinePoints.SetNum(2);
+
+	for (const float Freq : GridLineFrequencies)
+	{
+		// Draw vertical grid line:
+		const float GridLineLocalX = ScaleInfo.FrequencyToLocalX(Freq);
+		LinePoints[0] = { GridLineLocalX, LocalBackgroundRect.Top };
+		LinePoints[1] = { GridLineLocalX, LocalBackgroundRect.Bottom };
+		FSlateDrawElement::MakeLines(ElementList, LayerId, AllottedGeometry.ToPaintGeometry(), LinePoints, ESlateDrawEffect::None, LineColor);
+	}
+}
+
+void FAudioSpectrumPlotGridAndLabelDrawingHelper::DrawSoundLevelAxisLabels(const int32 LayerId, TConstArrayView<float> GridLineSoundLevels, const FSlateFontInfo& Font, const FLinearColor& TextColor)
+{
+	SpaceDrawSize = FontMeasureService->Measure(TEXT(" "), Font);
+
+	for (const float SoundLevel : GridLineSoundLevels)
+	{
+		const FString SoundLevelString = FormatSoundLevelString(SoundLevel);
+		const FVector2f LabelDrawSize = FontMeasureService->Measure(SoundLevelString, Font);
+		const float GridLineLocalY = ScaleInfo.SoundLevelToLocalY(SoundLevel);
+		const float LabelTop = GridLineLocalY - 0.5f * LabelDrawSize.Y;
+		const float LabelBottom = GridLineLocalY + 0.5f * LabelDrawSize.Y;
+		if (LabelTop >= LocalBackgroundRect.Top && LabelBottom <= LocalBackgroundRect.Bottom)
+		{
+			// Draw label on the left hand side:
+			DrawLabelIfNoOverlap(LayerId, LocalBackgroundRect.Left, LabelTop, LabelDrawSize, SoundLevelString, Font, TextColor);
+
+			// Draw label on the right hand side:
+			DrawLabelIfNoOverlap(LayerId, LocalBackgroundRect.Right - LabelDrawSize.X, LabelTop, LabelDrawSize, SoundLevelString, Font, TextColor);
+		}
+	}
+}
+
+void FAudioSpectrumPlotGridAndLabelDrawingHelper::DrawFrequencyAxisLabels(const int32 LayerId, TConstArrayView<float> GridLineFrequencies, const FSlateFontInfo& Font, const FLinearColor& TextColor)
+{
+	SpaceDrawSize = FontMeasureService->Measure(TEXT(" "), Font);
+
+	for (float Freq : GridLineFrequencies)
+	{
+		const FString FreqString = FormatFreqString(Freq);
+		const FVector2f LabelDrawSize = FontMeasureService->Measure(FreqString, Font);
+		const float GridLineLocalX = ScaleInfo.FrequencyToLocalX(Freq);
+		const float LabelLeft = GridLineLocalX - 0.5f * LabelDrawSize.X;
+		const float LabelRight = GridLineLocalX + 0.5f * LabelDrawSize.X;
+		if (LabelLeft >= LocalBackgroundRect.Left && LabelRight <= LocalBackgroundRect.Right)
+		{
+			// Draw label at the top:
+			DrawLabelIfNoOverlap(LayerId, LabelLeft, LocalBackgroundRect.Top, LabelDrawSize, FreqString, Font, TextColor);
+
+			// Draw label at the bottom:
+			DrawLabelIfNoOverlap(LayerId, LabelLeft, LocalBackgroundRect.Bottom - LabelDrawSize.Y, LabelDrawSize, FreqString, Font, TextColor);
+		}
+	}
+}
+
+FString FAudioSpectrumPlotGridAndLabelDrawingHelper::FormatSoundLevelString(const float SoundLevel)
+{
+	if (SoundLevel > 0.0f)
+	{
+		return FString::Printf(TEXT("+%.0f"), SoundLevel);
+	}
+	else
+	{
+		return FString::Printf(TEXT("%.0f"), SoundLevel);
+	}
+}
+
+FString FAudioSpectrumPlotGridAndLabelDrawingHelper::FormatFreqString(const float Freq)
+{
+	if (Freq >= 1000.0f)
+	{
+		return FString::Printf(TEXT("%.0fk"), Freq / 1000.0f);
+	}
+	else
+	{
+		return FString::Printf(TEXT("%.0f"), Freq);
+	}
+}
+
+void FAudioSpectrumPlotGridAndLabelDrawingHelper::DrawLabelIfNoOverlap(const int32 LayerId, const float LabelLeft, const float LabelTop, const FVector2f& LabelDrawSize, const FString LabelText, const FSlateFontInfo& Font, const FLinearColor& TextColor)
+{
+	const FSlateLayoutTransform LabelTransform(FVector2f(LabelLeft, LabelTop));
+	const FSlateRect LabelRect(LabelTransform.TransformPoint(FVector2f::ZeroVector), LabelTransform.TransformPoint(LabelDrawSize));
+	const FSlateRect ModifiedLabelRect = GetModifiedLabelRect(LabelRect);
+	if (!IsOverlappingPreviouslyDrawnLabel(ModifiedLabelRect))
+	{
+		const FPaintGeometry LabelPaintGeometry = AllottedGeometry.MakeChild(LabelDrawSize, LabelTransform).ToPaintGeometry();
+		FSlateDrawElement::MakeText(ElementList, LayerId, LabelPaintGeometry, LabelText, Font, ESlateDrawEffect::None, TextColor);
+		DrawnLabelRects.Add(LabelRect);
+	}
+}
+
+FSlateRect FAudioSpectrumPlotGridAndLabelDrawingHelper::GetModifiedLabelRect(const FSlateRect& LabelRect) const
+{
+	const float TightLabelTop = FMath::Lerp(LabelRect.Top, LabelRect.Bottom, 0.1f);
+	const float TightLabelBottom = FMath::Lerp(LabelRect.Bottom, LabelRect.Top, 0.1f);
+	const float PaddedLabelLeft = LabelRect.Left - 0.5 * SpaceDrawSize.X;
+	const float PaddedLabelRight = LabelRect.Right + 0.5 * SpaceDrawSize.X;
+	return FSlateRect(PaddedLabelLeft, TightLabelTop, PaddedLabelRight, TightLabelBottom);
+}
+
+bool FAudioSpectrumPlotGridAndLabelDrawingHelper::IsOverlappingPreviouslyDrawnLabel(const FSlateRect& LabelRect) const
+{
+	const FSlateRect* OverlappingDrawnLabel = DrawnLabelRects.FindByPredicate([&LabelRect](const FSlateRect& PrevDrawnLabelRect)
+		{
+			return (
+				LabelRect.Top < PrevDrawnLabelRect.Bottom && LabelRect.Bottom > PrevDrawnLabelRect.Top &&
+				LabelRect.Left < PrevDrawnLabelRect.Right && LabelRect.Right > PrevDrawnLabelRect.Left
+				);
+		});
+
+	return (OverlappingDrawnLabel != nullptr);
+};
+
+
+
+void SAudioSpectrumPlot::Construct(const FArguments& InArgs)
+{
+	check(InArgs._Style);
+
+	Style = InArgs._Style;
+	ViewMinFrequency = InArgs._ViewMinFrequency;
+	ViewMaxFrequency = InArgs._ViewMaxFrequency;
+	ViewMinSoundLevel = InArgs._ViewMinSoundLevel;
+	ViewMaxSoundLevel = InArgs._ViewMaxSoundLevel;
+	bDisplayFrequencyAxisLabels = InArgs._DisplayFrequencyAxisLabels;
+	bDisplaySoundLevelAxisLabels = InArgs._DisplaySoundLevelAxisLabels;
+	FrequencyAxisScale = InArgs._FrequencyAxisScale;
+	FrequencyAxisPixelBucketMode = InArgs._FrequencyAxisPixelBucketMode;
+	BackgroundColor = InArgs._BackgroundColor;
+	GridColor = InArgs._GridColor;
+	AxisLabelColor = InArgs._AxisLabelColor;
+	SpectrumColor = InArgs._SpectrumColor;
+	OnGetAudioSpectrumData = InArgs._OnGetAudioSpectrumData;
+}
+
+int32 SAudioSpectrumPlot::OnPaint(const FPaintArgs& Args, const FGeometry& AllottedGeometry, const FSlateRect& MyCullingRect, FSlateWindowElementList& OutDrawElements, int32 LayerId, const FWidgetStyle& InWidgetStyle, bool bParentEnabled) const
+{
+	const FAudioSpectrumPlotScaleInfo ScaleInfo(AllottedGeometry.GetLocalSize(), FrequencyAxisScale.Get(), ViewMinFrequency.Get(), ViewMaxFrequency.Get(), ViewMinSoundLevel.Get(), ViewMaxSoundLevel.Get());
+
+	LayerId = DrawSolidBackgroundRectangle(AllottedGeometry, OutDrawElements, LayerId, InWidgetStyle);
+
+	LayerId = DrawGridAndLabels(AllottedGeometry, OutDrawElements, LayerId, InWidgetStyle, ScaleInfo);
+
+	LayerId = DrawPowerSpectrum(AllottedGeometry, OutDrawElements, LayerId, InWidgetStyle, ScaleInfo);
+
+	return LayerId;
+}
+
+int32 SAudioSpectrumPlot::DrawSolidBackgroundRectangle(const FGeometry& AllottedGeometry, FSlateWindowElementList& OutDrawElements, int32 LayerId, const FWidgetStyle& InWidgetStyle) const
+{
+	const FSlateBrush BackgroundBrush;
+	const FLinearColor BoxColor = GetBackgroundColor(InWidgetStyle);
+	FSlateDrawElement::MakeBox(OutDrawElements, LayerId, AllottedGeometry.ToPaintGeometry(), &BackgroundBrush, ESlateDrawEffect::None, BoxColor);
+
+	return LayerId + 1;
+}
+
+int32 SAudioSpectrumPlot::DrawGridAndLabels(const FGeometry& AllottedGeometry, FSlateWindowElementList& OutDrawElements, int32 LayerId, const FWidgetStyle& InWidgetStyle, const FAudioSpectrumPlotScaleInfo& ScaleInfo) const
+{
+	TArray<float> GridLineSoundLevels;
+	GetGridLineSoundLevels(GridLineSoundLevels);
+
+	TArray<float> AllGridLineFrequencies;
+	TArray<float> MajorGridLineFrequencies;
+	GetGridLineFrequencies(AllGridLineFrequencies, MajorGridLineFrequencies);
+
+	const FSlateFontInfo& Font = Style->AxisLabelFont;
+	const FLinearColor LineColor = GetGridColor(InWidgetStyle);
+	const FLinearColor TextColor = GetAxisLabelColor(InWidgetStyle);
+
+	FAudioSpectrumPlotGridAndLabelDrawingHelper GridAndLabelDrawingHelper(AllottedGeometry, OutDrawElements, ScaleInfo);
+
+	GridAndLabelDrawingHelper.DrawSoundLevelGridLines(LayerId, GridLineSoundLevels, LineColor);
+	GridAndLabelDrawingHelper.DrawFrequencyGridLines(LayerId, AllGridLineFrequencies, LineColor);
+	LayerId++;
+
+	if (bDisplaySoundLevelAxisLabels.Get())
+	{
+		// Draw sound level axis labels for all grid lines.
+		GridAndLabelDrawingHelper.DrawSoundLevelAxisLabels(LayerId, GridLineSoundLevels, Font, TextColor);
+	}
+
+	if (bDisplayFrequencyAxisLabels.Get())
+	{
+		// Draw frequency axis labels for all major grid lines.
+		GridAndLabelDrawingHelper.DrawFrequencyAxisLabels(LayerId, MajorGridLineFrequencies, Font, TextColor);
+	}
+
+	if (GridAndLabelDrawingHelper.HasDrawnLabels())
+	{
+		// We drew some labels, so increment layer ID:
+		LayerId++;
+	}
+
+	return LayerId;
+}
+
+void SAudioSpectrumPlot::GetGridLineSoundLevels(TArray<float>& GridLineSoundLevels) const
+{
+	// Define grid line sound levels (dB scale):
+	const float MaxSoundLevel = ViewMaxSoundLevel.Get();
+	const float MinSoundLevel = ViewMinSoundLevel.Get();
+	const float SoundLevelIncrement = 20.0f * FMath::LogX(10.0f, 2.0f);
+
+	// Add grid lines from 0dB up to MaxSoundLevel:
+	float SoundLevel = 0.0f;
+	while (SoundLevel <= MaxSoundLevel)
+	{
+		GridLineSoundLevels.Add(SoundLevel);
+		SoundLevel += SoundLevelIncrement;
+	}
+
+	// Add grid lines from below 0dB down to MinSoundLevel:
+	SoundLevel = 0.0f - SoundLevelIncrement;
+	while (SoundLevel >= MinSoundLevel)
+	{
+		GridLineSoundLevels.Add(SoundLevel);
+		SoundLevel -= SoundLevelIncrement;
+	}
+}
+
+void SAudioSpectrumPlot::GetGridLineFrequencies(TArray<float>& AllGridLineFrequencies, TArray<float>& MajorGridLineFrequencies) const
+{
+	if (FrequencyAxisScale.Get() == EAudioSpectrumPlotFrequencyAxisScale::Logarithmic)
+	{
+		// Define grid line frequencies (log scale):
+
+		const float MinGridFreq = ViewMinFrequency.Get();
+		const float MaxGridFreq = ViewMaxFrequency.Get();
+		const float Log10MinGridFreq = FMath::LogX(10.0f, MinGridFreq);
+
+		float Freq = FMath::Pow(10.0f, FMath::Floor(Log10MinGridFreq));
+		while (Freq <= MaxGridFreq)
+		{
+			if (Freq >= MinGridFreq)
+			{
+				MajorGridLineFrequencies.Add(Freq);
+			}
+
+			const float FreqIncrement = Freq;
+			const float NextJump = 10.0f * FreqIncrement;
+			while (Freq < NextJump && Freq <= MaxGridFreq)
+			{
+				if (Freq >= MinGridFreq)
+				{
+					AllGridLineFrequencies.Add(Freq);
+				}
+
+				Freq += FreqIncrement;
+			}
+		}
+	}
+	else
+	{
+		// Define grid line frequencies (linear scale):
+		const float ViewFrequencyRange = ViewMaxFrequency.Get() - ViewMinFrequency.Get();
+
+		// Find grid spacing to draw around 10 grid lines:
+		const float Log10ApproxGridSpacing = FMath::LogX(10.0f, ViewFrequencyRange / 10.0f);
+		const float GridSpacing = FMath::Pow(10.0f, FMath::Floor(Log10ApproxGridSpacing));
+
+		// Add frequencies to the grid line arrays:
+		const float StartFrequency = GridSpacing * FMath::CeilToDouble(ViewMinFrequency.Get() / GridSpacing);
+		for (float Freq = StartFrequency; Freq <= ViewMaxFrequency.Get(); Freq += GridSpacing)
+		{
+			AllGridLineFrequencies.Add(Freq);
+			MajorGridLineFrequencies.Add(Freq);
+		}
+	}
+}
+
+int32 SAudioSpectrumPlot::DrawPowerSpectrum(const FGeometry& AllottedGeometry, FSlateWindowElementList& OutDrawElements, int32 LayerId, const FWidgetStyle& InWidgetStyle, const FAudioSpectrumPlotScaleInfo& ScaleInfo) const
+{
+	// Get the power spectrum data if available:
+	const FAudioPowerSpectrumData PowerSpectrum = GetPowerSpectrum();
+	ensure(PowerSpectrum.CenterFrequencies.Num() == PowerSpectrum.SquaredMagnitudes.Num());
+	const int NumFrequencies = FMath::Min(PowerSpectrum.CenterFrequencies.Num(), PowerSpectrum.SquaredMagnitudes.Num());
+	if (NumFrequencies > 0)
+	{
+		// Convert to array of data points with X == frequency, Y == sound level in dB.
+		TArray<FVector2f> DataPoints;
+		DataPoints.Reserve(NumFrequencies);
+
+		const float ClampMinFrequency = (FrequencyAxisScale.Get() == EAudioSpectrumPlotFrequencyAxisScale::Logarithmic) ? 0.00001f : -FLT_MAX; // Cannot plot DC with log scale.
+		const float ClampMinMagnitudeSquared = FMath::Pow(10.0f, -200.0f / 10.0f); // Clamp at -200dB
+		for (int Index = 0; Index < NumFrequencies; Index++)
+		{
+			const float Frequency = FMath::Max(PowerSpectrum.CenterFrequencies[Index], ClampMinFrequency);
+			const float MagnitudeSquared = FMath::Max(PowerSpectrum.SquaredMagnitudes[Index], ClampMinMagnitudeSquared);
+			const float SoundLevel = 10.0f * FMath::LogX(10.0f, MagnitudeSquared);
+			DataPoints.Add({ Frequency, SoundLevel });
+		}
+
+		// Line points to plot will be added to this array:
+		TArray<FVector2f> LinePoints;
+
+		switch (FrequencyAxisPixelBucketMode.Get())
+		{
+			case EAudioSpectrumPlotFrequencyAxisPixelBucketMode::Sample:
+			{
+				// For DataPoints that map to the same frequency axis pixel bucket, choose the one that is nearest the bucket center:
+				const auto CostFunction = [&ScaleInfo](const FVector2f& DataPoint)
+				{
+					const float LocalX = ScaleInfo.FrequencyToLocalX(DataPoint.X);
+					return FMath::Abs(LocalX - FMath::RoundToInt32(LocalX));
+				};
+
+				// Get the line points to plot:
+				GetSpectrumLinePoints(LinePoints, DataPoints, ScaleInfo, CostFunction);
+			}
+			break;
+			case EAudioSpectrumPlotFrequencyAxisPixelBucketMode::Peak:
+			{
+				// For DataPoints that map to the same frequency axis pixel bucket, choose the one with the highest sound level:
+				const auto CostFunction = [](const FVector2f& DataPoint) { return -DataPoint.Y; };
+
+				// Get the line points to plot:
+				GetSpectrumLinePoints(LinePoints, DataPoints, ScaleInfo, CostFunction);
+			}
+			break;
+			case EAudioSpectrumPlotFrequencyAxisPixelBucketMode::Average:
+			{
+				// For DataPoints that map to the same frequency axis pixel bucket, take the average:
+
+				int32 CurrFreqAxisPixelBucket = INT32_MIN;
+				FVector2f CurrSum = FVector2f::ZeroVector;
+				int CurrCount = 0;
+				for (FVector2f DataPoint : DataPoints)
+				{
+					const float LocalX = ScaleInfo.FrequencyToLocalX(DataPoint.X);
+					const float LocalY = ScaleInfo.SoundLevelToLocalY(DataPoint.Y);
+
+					const int32 FreqAxisPixelBucket = FMath::RoundToInt32(LocalX);
+					if (FreqAxisPixelBucket != CurrFreqAxisPixelBucket && CurrCount > 0)
+					{
+						// New DataPoint is not at the same frequency axis pixel bucket.
+
+						// Add current average to line plot:
+						LinePoints.Add(CurrSum / CurrCount);
+
+						// Reset current average:
+						CurrSum = FVector2f::ZeroVector;
+						CurrCount = 0;
+					}
+
+					// Set the current frequency axis pixel bucket, and add to the average:
+					CurrFreqAxisPixelBucket = FreqAxisPixelBucket;
+					CurrSum += FVector2f(LocalX, LocalY);
+					CurrCount++;
+				}
+
+				// Add remaining average to line plot:
+				check(CurrCount > 0);
+				LinePoints.Add(CurrSum / CurrCount);
+			}
+			break;
+		}
+
+		// Actually draw the line points:
+		const FLinearColor& LineColor = GetSpectrumColor(InWidgetStyle);
+		FSlateDrawElement::MakeLines(OutDrawElements, LayerId, AllottedGeometry.ToPaintGeometry(), LinePoints, ESlateDrawEffect::None, LineColor, true, 1.0f);
+		LayerId++;
+	}
+
+	return LayerId;
+}
+
+FAudioPowerSpectrumData SAudioSpectrumPlot::GetPowerSpectrum() const
+{
+	if (OnGetAudioSpectrumData.IsBound())
+	{
+		return OnGetAudioSpectrumData.Execute();
+	}
+
+	return FAudioPowerSpectrumData();
+}
+
+void SAudioSpectrumPlot::GetSpectrumLinePoints(TArray<FVector2f>& OutLinePoints, TConstArrayView<FVector2f> DataPoints, const FAudioSpectrumPlotScaleInfo& ScaleInfo, TFunctionRef<float(const FVector2f& DataPoint)> CostFunction)
+{
+	// Function to find whether two data points map to the same frequency axis pixel bucket:
+	const auto IsSameFreqAxisPixelBucket = [&ScaleInfo](FVector2f DataPoint1, FVector2f DataPoint2)
+	{
+		const int32 PixelBucket1 = FMath::RoundToInt32(ScaleInfo.FrequencyToLocalX(DataPoint1.X));
+		const int32 PixelBucket2 = FMath::RoundToInt32(ScaleInfo.FrequencyToLocalX(DataPoint2.X));
+		return (PixelBucket1 == PixelBucket2);
+	};
+
+	TOptional<FVector2f> CurrBestDataPoint;
+
+	for (FVector2f DataPoint : DataPoints)
+	{
+		if (CurrBestDataPoint.IsSet() && !IsSameFreqAxisPixelBucket(DataPoint, CurrBestDataPoint.GetValue()))
+		{
+			// New DataPoint is not at the same frequency axis pixel bucket as CurrBestDataPoint.
+
+			// Add CurrBestDataPoint to line plot:
+			const FVector2f LocalPosCurrBestDataPoint = ScaleInfo.ToLocalPos(CurrBestDataPoint.GetValue());
+			OutLinePoints.Add(LocalPosCurrBestDataPoint);
+
+			// Reset best value: 
+			CurrBestDataPoint.Reset();
+		}
+
+		if (!CurrBestDataPoint.IsSet() || CostFunction(DataPoint) < CostFunction(CurrBestDataPoint.GetValue()))
+		{
+			// New DataPoint is either at a new frequency axis pixel bucket or is better than CurrBestDataPoint.
+			CurrBestDataPoint = DataPoint;
+		}
+	}
+
+	{
+		// Add final CurrBestDataPoint to line plot:
+		const FVector2f LocalPosCurrBestDataPoint = ScaleInfo.ToLocalPos(CurrBestDataPoint.GetValue());
+		OutLinePoints.Add(LocalPosCurrBestDataPoint);
+	}
+}
+
+FLinearColor SAudioSpectrumPlot::GetBackgroundColor(const FWidgetStyle& InWidgetStyle) const
+{
+	const FSlateColor& SlateColor = (BackgroundColor.Get() != FSlateColor::UseStyle()) ? BackgroundColor.Get() : Style->BackgroundColor;
+	return SlateColor.GetColor(InWidgetStyle) * InWidgetStyle.GetColorAndOpacityTint();
+}
+
+FLinearColor SAudioSpectrumPlot::GetGridColor(const FWidgetStyle& InWidgetStyle) const
+{
+	const FSlateColor& SlateColor = (GridColor.Get() != FSlateColor::UseStyle()) ? GridColor.Get() : Style->GridColor;
+	return SlateColor.GetColor(InWidgetStyle) * InWidgetStyle.GetColorAndOpacityTint();
+}
+
+FLinearColor SAudioSpectrumPlot::GetAxisLabelColor(const FWidgetStyle& InWidgetStyle) const
+{
+	const FSlateColor& SlateColor = (AxisLabelColor.Get() != FSlateColor::UseStyle()) ? AxisLabelColor.Get() : Style->AxisLabelColor;
+	return SlateColor.GetColor(InWidgetStyle) * InWidgetStyle.GetColorAndOpacityTint();
+}
+
+FLinearColor SAudioSpectrumPlot::GetSpectrumColor(const FWidgetStyle& InWidgetStyle) const
+{
+	const FSlateColor& SlateColor = (SpectrumColor.Get() != FSlateColor::UseStyle()) ? SpectrumColor.Get() : Style->SpectrumColor;
+	return SlateColor.GetColor(InWidgetStyle) * InWidgetStyle.GetColorAndOpacityTint();
+}
+END_SLATE_FUNCTION_BUILD_OPTIMIZATION
diff --git a/Engine/Plugins/Runtime/AudioWidgets/Source/AudioWidgets/Public/AudioSpectrumPlotStyle.h b/Engine/Plugins/Runtime/AudioWidgets/Source/AudioWidgets/Public/AudioSpectrumPlotStyle.h
new file mode 100644
--- /dev/null
+++ b/Engine/Plugins/Runtime/AudioWidgets/Source/AudioWidgets/Public/AudioSpectrumPlotStyle.h
@@ -0,0 +1,47 @@
+// Copyright Epic Games, Inc. All Rights Reserved.
+
+#pragma once
+
+#include "Fonts/SlateFontInfo.h"
+#include "Styling/StyleColors.h"
+#include "Styling/StyleDefaults.h"
+#include "Styling/SlateWidgetStyle.h"
+#include "AudioSpectrumPlotStyle.generated.h"
+
+/**
+ * Represents the appearance of an SAudioSpectrumPlot
+ */
+USTRUCT(BlueprintType)
+struct AUDIOWIDGETS_API FAudioSpectrumPlotStyle : public FSlateWidgetStyle
+{
+	GENERATED_USTRUCT_BODY()
+
+	FAudioSpectrumPlotStyle();
+
+	virtual ~FAudioSpectrumPlotStyle() {}
+
+	static const FName TypeName;
+	virtual const FName GetTypeName() const override { return TypeName; };
+
+	static const FAudioSpectrumPlotStyle& GetDefault();
+
+	UPROPERTY(EditAnywhere, BlueprintReadWrite, Category = Appearance)
+	FSlateColor BackgroundColor = FStyleColors::Background;
+	FAudioSpectrumPlotStyle& SetBackgroundColor(const FSlateColor& InBackgroundColor) { BackgroundColor = InBackgroundColor; return *this; }
+
+	UPROPERTY(EditAnywhere, BlueprintReadWrite, Category = Appearance)
+	FSlateColor GridColor = FMath::Lerp(FStyleColors::Background.GetSpecifiedColor(), FStyleColors::Foreground.GetSpecifiedColor(), 0.1f);
+	FAudioSpectrumPlotStyle& SetGridColor(const FSlateColor& InGridColor) { GridColor = InGridColor; return *this; }
+
+	UPROPERTY(EditAnywhere, BlueprintReadWrite, Category = Appearance)
+	FSlateColor AxisLabelColor = FStyleColors::Foreground;
+	FAudioSpectrumPlotStyle& SetAxisLabelColor(const FSlateColor& InAxisLabelColor) { AxisLabelColor = InAxisLabelColor; return *this; }
+
+	UPROPERTY(EditAnywhere, BlueprintReadWrite, Category = Appearance)
+	FSlateFontInfo AxisLabelFont = FStyleDefaults::GetFontInfo(8);
+	FAudioSpectrumPlotStyle& SetAxisLabelFont(const FSlateFontInfo& InAxisLabelFont) { AxisLabelFont = InAxisLabelFont; return *this; }
+
+	UPROPERTY(EditAnywhere, BlueprintReadWrite, Category = Appearance)
+	FSlateColor SpectrumColor = FStyleColors::Foreground;
+	FAudioSpectrumPlotStyle& SetSpectrumColor(const FSlateColor& InSpectrumColor) { SpectrumColor = InSpectrumColor; return *this; }
+};
diff --git a/Engine/Plugins/Runtime/AudioWidgets/Source/AudioWidgets/Public/SAudioSpectrumPlot.h b/Engine/Plugins/Runtime/AudioWidgets/Source/AudioWidgets/Public/SAudioSpectrumPlot.h
new file mode 100644
--- /dev/null
+++ b/Engine/Plugins/Runtime/AudioWidgets/Source/AudioWidgets/Public/SAudioSpectrumPlot.h
@@ -0,0 +1,194 @@
+// Copyright Epic Games, Inc. All Rights Reserved.
+
+#pragma once
+
+#include "AudioSpectrumPlotStyle.h"
+#include "Styling/ISlateStyle.h"
+#include "Styling/SlateWidgetStyleAsset.h"
+#include "Widgets/SCompoundWidget.h"
+
+UENUM(BlueprintType)
+enum class EAudioSpectrumPlotFrequencyAxisScale
+{
+	Linear,
+	Logarithmic,
+};
+
+UENUM(BlueprintType)
+enum class EAudioSpectrumPlotFrequencyAxisPixelBucketMode
+{
+	Sample UMETA(ToolTip = "Plot one data point per frequency axis pixel bucket only, choosing the data point nearest the pixel center."),
+	Peak UMETA(ToolTip = "Plot one data point per frequency axis pixel bucket only, choosing the data point with the highest sound level."),
+	Average UMETA(ToolTip = "Plot the average of the data points in each frequency axis pixel bucket."),
+};
+
+/**
+ * Utility class for converting between spectrum data and local/absolute screen space.
+ */
+class FAudioSpectrumPlotScaleInfo
+{
+public:
+	FAudioSpectrumPlotScaleInfo(const FVector2f InWidgetSize, EAudioSpectrumPlotFrequencyAxisScale InFrequencyAxisScale, float InViewMinFrequency, float InViewMaxFrequency, float InViewMinSoundLevel, float InViewMaxSoundLevel)
+		: WidgetSize(InWidgetSize)
+		, FrequencyAxisScale(InFrequencyAxisScale)
+		, TransformedViewMinFrequency(ForwardTransformFrequency(InViewMinFrequency))
+		, TransformedViewMaxFrequency(ForwardTransformFrequency(InViewMaxFrequency))
+		, TransformedViewFrequencyRange(TransformedViewMaxFrequency - TransformedViewMinFrequency)
+		, PixelsPerTransformedHz((TransformedViewFrequencyRange > 0.0f) ? (InWidgetSize.X / TransformedViewFrequencyRange) : 0.0f)
+		, ViewMinSoundLevel(InViewMinSoundLevel)
+		, ViewMaxSoundLevel(InViewMaxSoundLevel)
+		, ViewSoundLevelRange(InViewMaxSoundLevel - InViewMinSoundLevel)
+		, PixelsPerDecibel((ViewSoundLevelRange > 0.0f) ? (InWidgetSize.Y / ViewSoundLevelRange) : 0.0f)
+	{
+		//
+	}
+
+	float LocalXToFrequency(float ScreenX) const
+	{
+		const float TransformedFrequency = (PixelsPerTransformedHz != 0.0f) ? ((ScreenX / PixelsPerTransformedHz) + TransformedViewMinFrequency) : 0.0f;
+		return InverseTransformFrequency(TransformedFrequency);
+	}
+
+	float FrequencyToLocalX(float Frequency) const
+	{
+		return (ForwardTransformFrequency(Frequency) - TransformedViewMinFrequency) * PixelsPerTransformedHz;
+	}
+
+	float LocalYToSoundLevel(float ScreenY) const
+	{
+		return (PixelsPerDecibel != 0.0f) ? (ViewMaxSoundLevel - (ScreenY / PixelsPerDecibel)) : 0.0f;
+	}
+
+	float SoundLevelToLocalY(float SoundLevel) const
+	{
+		return (ViewMaxSoundLevel - SoundLevel) * PixelsPerDecibel;
+	}
+
+	FVector2f ToLocalPos(const FVector2f& FrequencyAndSoundLevel) const
+	{
+		return { FrequencyToLocalX(FrequencyAndSoundLevel.X), SoundLevelToLocalY(FrequencyAndSoundLevel.Y) };
+	}
+
+private:
+	float ForwardTransformFrequency(float Frequency) const
+	{
+		return (FrequencyAxisScale == EAudioSpectrumPlotFrequencyAxisScale::Logarithmic) ? FMath::Loge(Frequency) : Frequency;
+	}
+
+	float InverseTransformFrequency(float TransformedFrequency) const
+	{
+		return (FrequencyAxisScale == EAudioSpectrumPlotFrequencyAxisScale::Logarithmic) ? FMath::Exp(TransformedFrequency) : TransformedFrequency;
+	}
+
+	const FVector2f WidgetSize;
+
+	const EAudioSpectrumPlotFrequencyAxisScale FrequencyAxisScale;
+	const float TransformedViewMinFrequency;
+	const float TransformedViewMaxFrequency;
+	const float TransformedViewFrequencyRange;
+	const float PixelsPerTransformedHz;
+
+	const float ViewMinSoundLevel;
+	const float ViewMaxSoundLevel;
+	const float ViewSoundLevelRange;
+	const float PixelsPerDecibel;
+};
+
+/**
+ * The audio spectrum data to plot.
+ */
+struct FAudioPowerSpectrumData
+{
+	TConstArrayView<float> CenterFrequencies;
+	TConstArrayView<float> SquaredMagnitudes;
+};
+
+DECLARE_DELEGATE_RetVal(FAudioPowerSpectrumData, FGetAudioSpectrumData);
+
+/**
+ * Slate Widget for plotting an audio power spectrum, with linear or log frequency scale, and decibels sound levels.
+ */
+class AUDIOWIDGETS_API SAudioSpectrumPlot : public SCompoundWidget
+{
+public:
+	SLATE_BEGIN_ARGS(SAudioSpectrumPlot)
+		: _Style(&FAudioSpectrumPlotStyle::GetDefault())
+		, _ViewMinFrequency(20.0f)
+		, _ViewMaxFrequency(20000.0f)
+		, _ViewMinSoundLevel(-60.0f)
+		, _ViewMaxSoundLevel(12.0f)
+		, _DisplayFrequencyAxisLabels(true)
+		, _DisplaySoundLevelAxisLabels(true)
+		, _FrequencyAxisScale(EAudioSpectrumPlotFrequencyAxisScale::Logarithmic)
+		, _FrequencyAxisPixelBucketMode(EAudioSpectrumPlotFrequencyAxisPixelBucketMode::Average)
+		, _BackgroundColor(FSlateColor::UseStyle())
+		, _GridColor(FSlateColor::UseStyle())
+		, _AxisLabelColor(FSlateColor::UseStyle())
+		, _SpectrumColor(FSlateColor::UseStyle())
+	{}
+		SLATE_STYLE_ARGUMENT(FAudioSpectrumPlotStyle, Style)
+		SLATE_ATTRIBUTE(float, ViewMinFrequency)
+		SLATE_ATTRIBUTE(float, ViewMaxFrequency)
+		SLATE_ATTRIBUTE(float, ViewMinSoundLevel)
+		SLATE_ATTRIBUTE(float, ViewMaxSoundLevel)
+		SLATE_ATTRIBUTE(bool, DisplayFrequencyAxisLabels)
+		SLATE_ATTRIBUTE(bool, DisplaySoundLevelAxisLabels)
+		SLATE_ATTRIBUTE(EAudioSpectrumPlotFrequencyAxisScale, FrequencyAxisScale)
+		SLATE_ATTRIBUTE(EAudioSpectrumPlotFrequencyAxisPixelBucketMode, FrequencyAxisPixelBucketMode)
+		SLATE_ATTRIBUTE(FSlateColor, BackgroundColor)
+		SLATE_ATTRIBUTE(FSlateColor, GridColor)
+		SLATE_ATTRIBUTE(FSlateColor, AxisLabelColor)
+		SLATE_ATTRIBUTE(FSlateColor, SpectrumColor)
+		SLATE_EVENT(FGetAudioSpectrumData, OnGetAudioSpectrumData)
+	SLATE_END_ARGS()
+
+	/** Constructs this widget with InArgs */
+	void Construct(const FArguments& InArgs);
+
+	void SetViewMinFrequency(float InViewMinFrequency) { ViewMinFrequency = InViewMinFrequency; }
+	void SetViewMaxFrequency(float InViewMaxFrequency) { ViewMaxFrequency = InViewMaxFrequency; }
+	void SetViewMinSoundLevel(float InViewMinSoundLevel) { ViewMinSoundLevel = InViewMinSoundLevel; }
+	void SetViewMaxSoundLevel(float InViewMaxSoundLevel) { ViewMaxSoundLevel = InViewMaxSoundLevel; }
+	void SetDisplayFrequencyAxisLabels(bool bInDisplayFrequencyAxisLabels) { bDisplayFrequencyAxisLabels = bInDisplayFrequencyAxisLabels; }
+	void SetDisplaySoundLevelAxisLabels(bool bInDisplaySoundLevelAxisLabels) { bDisplaySoundLevelAxisLabels = bInDisplaySoundLevelAxisLabels; }
+	void SetFrequencyAxisScale(EAudioSpectrumPlotFrequencyAxisScale InFrequencyAxisScale) { FrequencyAxisScale = InFrequencyAxisScale; }
+	void SetFrequencyAxisPixelBucketMode(EAudioSpectrumPlotFrequencyAxisPixelBucketMode InFrequencyAxisPixelBucketMode) { FrequencyAxisPixelBucketMode = InFrequencyAxisPixelBucketMode; }
+
+private:
+	// Begin SWidget overrides.
+	virtual int32 OnPaint(const FPaintArgs& Args, const FGeometry& AllottedGeometry, const FSlateRect& MyCullingRect, FSlateWindowElementList& OutDrawElements, int32 LayerId, const FWidgetStyle& InWidgetStyle, bool bParentEnabled) const override;
+	// End SWidget overrides.
+
+	int32 DrawSolidBackgroundRectangle(const FGeometry& AllottedGeometry, FSlateWindowElementList& OutDrawElements, int32 LayerId, const FWidgetStyle& InWidgetStyle) const;
+
+	int32 DrawGridAndLabels(const FGeometry& AllottedGeometry, FSlateWindowElementList& OutDrawElements, int32 LayerId, const FWidgetStyle& InWidgetStyle, const FAudioSpectrumPlotScaleInfo& ScaleInfo) const;
+	void GetGridLineSoundLevels(TArray<float>& GridLineSoundLevels) const;
+	void GetGridLineFrequencies(TArray<float>& AllGridLineFrequencies, TArray<float>& MajorGridLineFrequencies) const;
+
+	int32 DrawPowerSpectrum(const FGeometry& AllottedGeometry, FSlateWindowElementList& OutDrawElements, int32 LayerId, const FWidgetStyle& InWidgetStyle, const FAudioSpectrumPlotScaleInfo& ScaleInfo) const;
+	FAudioPowerSpectrumData GetPowerSpectrum() const;
+
+	// This is a function to reduce the given array of data points to a possibly shorter array of points that will form the line to be plotted.
+	// Where multiple data points map to the same frequency axis pixel bucket, the given 'cost function' will be used to select the best data point (the data point with the lowest 'cost').
+	static void GetSpectrumLinePoints(TArray<FVector2f>& OutLinePoints, TConstArrayView<FVector2f> DataPoints, const FAudioSpectrumPlotScaleInfo& ScaleInfo, TFunctionRef<float(const FVector2f& DataPoint)> CostFunction);
+
+	FLinearColor GetBackgroundColor(const FWidgetStyle& InWidgetStyle) const;
+	FLinearColor GetGridColor(const FWidgetStyle& InWidgetStyle) const;
+	FLinearColor GetAxisLabelColor(const FWidgetStyle& InWidgetStyle) const;
+	FLinearColor GetSpectrumColor(const FWidgetStyle& InWidgetStyle) const;	
+
+	const FAudioSpectrumPlotStyle* Style;
+	TAttribute<float> ViewMinFrequency;
+	TAttribute<float> ViewMaxFrequency;
+	TAttribute<float> ViewMinSoundLevel;
+	TAttribute<float> ViewMaxSoundLevel;
+	TAttribute<bool> bDisplayFrequencyAxisLabels;
+	TAttribute<bool> bDisplaySoundLevelAxisLabels;
+	TAttribute<EAudioSpectrumPlotFrequencyAxisScale> FrequencyAxisScale;
+	TAttribute<EAudioSpectrumPlotFrequencyAxisPixelBucketMode> FrequencyAxisPixelBucketMode;
+	TAttribute<FSlateColor> BackgroundColor;
+	TAttribute<FSlateColor> GridColor;
+	TAttribute<FSlateColor> AxisLabelColor;
+	TAttribute<FSlateColor> SpectrumColor;
+	FGetAudioSpectrumData OnGetAudioSpectrumData;
+};
