Added FAudioSpectrumAnalyzer class to AudioWidgets. This class owns a ConstantQ analyzer and a corresponding Slate widget for displaying the resulting spectrum. Exponential time-smoothing is applied to the spectrum.
Added an instance of this to analyzers tab in MetaSoundEditor.

diff --git a/Engine/Plugins/Runtime/AudioWidgets/Source/AudioWidgets/Private/AudioSpectrumAnalyzer.cpp b/Engine/Plugins/Runtime/AudioWidgets/Source/AudioWidgets/Private/AudioSpectrumAnalyzer.cpp
new file mode 100644
--- /dev/null
+++ b/Engine/Plugins/Runtime/AudioWidgets/Source/AudioWidgets/Private/AudioSpectrumAnalyzer.cpp
@@ -0,0 +1,131 @@
+// Copyright Epic Games, Inc. All Rights Reserved.
+
+#include "AudioSpectrumAnalyzer.h"
+
+#include "DSP/EnvelopeFollower.h"
+#include "Engine/World.h"
+
+namespace AudioWidgets
+{
+	FAudioSpectrumAnalyzer::FAudioSpectrumAnalyzer(int32 InNumChannels, UWorld& InWorld, TObjectPtr<UAudioBus> InExternalAudioBus)
+		: Widget(SNew(SAudioSpectrumPlot)
+			.Clipping(EWidgetClipping::ClipToBounds)
+			.DisplayFrequencyAxisLabels(false)
+			.DisplaySoundLevelAxisLabels(false)
+			.OnGetAudioSpectrumData_Raw(this, &FAudioSpectrumAnalyzer::GetAudioSpectrumData))
+	{
+		Init(InNumChannels, InWorld, InExternalAudioBus);
+	}
+
+	FAudioSpectrumAnalyzer::~FAudioSpectrumAnalyzer()
+	{
+		Teardown();
+	}
+
+	UAudioBus* FAudioSpectrumAnalyzer::GetAudioBus() const
+	{
+		return AudioBus.Get();
+	}
+
+	TSharedRef<SWidget> FAudioSpectrumAnalyzer::GetWidget() const
+	{
+		return Widget->AsShared();
+	}
+
+	void FAudioSpectrumAnalyzer::Init(int32 InNumChannels, UWorld& InWorld, TObjectPtr<UAudioBus> InExternalAudioBus)
+	{
+		check(InNumChannels > 0);
+
+		Teardown();
+
+		Settings = TStrongObjectPtr(NewObject<UConstantQSettings>());
+		Settings->SpectrumType = EAudioSpectrumType::PowerSpectrum;
+		Settings->NumBandsPerOctave = 6.0f;
+		Settings->NumBands = 61;
+		Settings->StartingFrequencyHz = 20000.0f * FMath::Pow(0.5f, (Settings->NumBands - 1) / Settings->NumBandsPerOctave);
+		Settings->FFTSize = EConstantQFFTSizeEnum::XXLarge;
+		Settings->bDownmixToMono = true;
+		Settings->BandWidthStretch = 2.0f;
+
+		Analyzer = TStrongObjectPtr(NewObject<UConstantQAnalyzer>());
+		Analyzer->Settings = Settings.Get();
+
+		bUseExternalAudioBus = InExternalAudioBus != nullptr;
+
+		AudioBus = bUseExternalAudioBus ? TStrongObjectPtr(InExternalAudioBus.Get()) : TStrongObjectPtr(NewObject<UAudioBus>());
+		AudioBus->AudioBusChannels = EAudioBusChannels(InNumChannels - 1);
+
+		ResultsDelegateHandle = Analyzer->OnConstantQResultsNative.AddRaw(this, &FAudioSpectrumAnalyzer::OnConstantQResults);
+
+		WorldPtr = &InWorld;
+		Analyzer->StartAnalyzing(&InWorld, AudioBus.Get());
+	}
+
+	void FAudioSpectrumAnalyzer::OnConstantQResults(UConstantQAnalyzer* InSpectrumAnalyzer, int32 ChannelIndex, const TArray<FConstantQResults>& InSpectrumResultsArray)
+	{
+		if (InSpectrumAnalyzer == Analyzer.Get())
+		{
+			for (const FConstantQResults& SpectrumResults : InSpectrumResultsArray)
+			{
+				if (PrevTimeStamp.IsSet() && SpectrumResults.TimeSeconds > PrevTimeStamp.GetValue())
+				{
+					// Calculate AR smoother coefficients:
+					const float DeltaT = (SpectrumResults.TimeSeconds - PrevTimeStamp.GetValue());
+					Audio::FAttackRelease AttackRelease(1.0f / DeltaT, AttackTimeMsec, ReleaseTimeMsec, true);
+
+					// Apply AR smoothing for each frequency:
+					check(SpectrumResults.SpectrumValues.Num() == ARSmoothedSquaredMagnitudes.Num());
+					for (int Index = 0; Index < SpectrumResults.SpectrumValues.Num(); Index++)
+					{
+						const float OldValue = ARSmoothedSquaredMagnitudes[Index];
+						const float NewValue = SpectrumResults.SpectrumValues[Index];
+						const float ARSmootherCoefficient = (NewValue >= OldValue) ? AttackRelease.GetAttackTimeSamples() : AttackRelease.GetReleaseTimeSamples();
+						const float SmoothedValue = FMath::Lerp(NewValue, OldValue, ARSmootherCoefficient);
+						ARSmoothedSquaredMagnitudes[Index] = Audio::UnderflowClamp(SmoothedValue);
+					}
+				}
+				else
+				{
+					// Init center frequencies:
+					CenterFrequencies.SetNumUninitialized(Analyzer->GetNumCenterFrequencies());
+					Analyzer->GetCenterFrequencies(CenterFrequencies);
+
+					// Init spectrum data:
+					ARSmoothedSquaredMagnitudes = SpectrumResults.SpectrumValues;
+				}
+
+				PrevTimeStamp = SpectrumResults.TimeSeconds;
+			}
+		}
+	}
+
+	void FAudioSpectrumAnalyzer::Teardown()
+	{
+		if (Analyzer.IsValid() && Analyzer->IsValidLowLevel())
+		{
+			Analyzer->StopAnalyzing();
+			if (ResultsDelegateHandle.IsValid())
+			{
+				Analyzer->OnConstantQResultsNative.Remove(ResultsDelegateHandle);
+			}
+			
+			Analyzer.Reset();
+		}
+
+		ResultsDelegateHandle.Reset();
+		PrevTimeStamp.Reset();
+		CenterFrequencies.Empty();
+		ARSmoothedSquaredMagnitudes.Empty();
+
+		AudioBus.Reset();
+		Settings.Reset();
+
+		bUseExternalAudioBus = false;
+	}
+
+	FAudioPowerSpectrumData FAudioSpectrumAnalyzer::GetAudioSpectrumData() const
+	{
+		check(CenterFrequencies.Num() == ARSmoothedSquaredMagnitudes.Num());
+		return FAudioPowerSpectrumData{ CenterFrequencies, ARSmoothedSquaredMagnitudes };
+	}
+} // namespace AudioWidgets
diff --git a/Engine/Plugins/Runtime/AudioWidgets/Source/AudioWidgets/Public/AudioSpectrumAnalyzer.h b/Engine/Plugins/Runtime/AudioWidgets/Source/AudioWidgets/Public/AudioSpectrumAnalyzer.h
new file mode 100644
--- /dev/null
+++ b/Engine/Plugins/Runtime/AudioWidgets/Source/AudioWidgets/Public/AudioSpectrumAnalyzer.h
@@ -0,0 +1,66 @@
+// Copyright Epic Games, Inc. All Rights Reserved.
+#pragma once
+
+#include "ConstantQ.h"
+#include "SAudioSpectrumPlot.h"
+#include "Sound/AudioBus.h"
+#include "UObject/StrongObjectPtr.h"
+
+class UWorld;
+
+namespace AudioWidgets
+{
+	/**
+	 * Owns an analyzer and a corresponding Slate widget for displaying the resulting spectrum.
+	 * Exponential time-smoothing is applied to the spectrum.
+	 * Can either create an Audio Bus to analyze, or analyze the given Bus.
+	 */
+	class AUDIOWIDGETS_API FAudioSpectrumAnalyzer
+	{
+	public:
+		FAudioSpectrumAnalyzer(int32 InNumChannels, UWorld& InWorld, TObjectPtr<UAudioBus> InExternalAudioBus = nullptr);
+		~FAudioSpectrumAnalyzer();
+
+		UAudioBus* GetAudioBus() const;
+
+		TSharedRef<SWidget> GetWidget() const;
+
+		void Init(int32 InNumChannels, UWorld& InWorld, TObjectPtr<UAudioBus> InExternalAudioBus = nullptr);
+
+	protected:
+		void OnConstantQResults(UConstantQAnalyzer* InSpectrumAnalyzer, int32 ChannelIndex, const TArray<FConstantQResults>& InSpectrumResultsArray);
+		FAudioPowerSpectrumData GetAudioSpectrumData() const;
+
+	private:
+		void Teardown();
+
+		/** Audio analyzer object. */
+		TStrongObjectPtr<UConstantQAnalyzer> Analyzer;
+
+		/** The audio bus used for analysis. */
+		TStrongObjectPtr<UAudioBus> AudioBus;
+
+		/** Meaning of spectrum data. */
+		TArray<float> CenterFrequencies;
+
+		/** Cached spectrum data, with AR smoothing applied. */
+		TArray<float> ARSmoothedSquaredMagnitudes;
+
+		/** Handle for results delegate for analyzer. */
+		FDelegateHandle ResultsDelegateHandle;
+
+		/** Analyzer settings. */
+		TStrongObjectPtr<UConstantQSettings> Settings;
+
+		/** Slate widget for spectrum display */
+		TSharedPtr<SAudioSpectrumPlot> Widget;
+
+		TWeakObjectPtr<UWorld> WorldPtr;
+
+		bool bUseExternalAudioBus = false;
+
+		TOptional<float> PrevTimeStamp;
+		float AttackTimeMsec = 300.0f;
+		float ReleaseTimeMsec = 300.0f;
+	};
+} // namespace AudioWidgets
diff --git a/Engine/Plugins/Runtime/Metasound/Source/MetasoundEditor/Private/MetasoundEditor.cpp b/Engine/Plugins/Runtime/Metasound/Source/MetasoundEditor/Private/MetasoundEditor.cpp
--- a/Engine/Plugins/Runtime/Metasound/Source/MetasoundEditor/Private/MetasoundEditor.cpp
+++ b/Engine/Plugins/Runtime/Metasound/Source/MetasoundEditor/Private/MetasoundEditor.cpp
@@ -6,6 +6,7 @@
 #include "AudioDevice.h"
 #include "AudioMeterStyle.h"
 #include "AudioOscilloscope.h"
+#include "AudioSpectrumAnalyzer.h"
 #include "AudioVectorscope.h"
 #include "AudioWidgetsEnums.h"
 #include "Components/AudioComponent.h"
@@ -682,7 +683,7 @@ namespace Metasound
 
 		TSharedPtr<SWidget> FEditor::BuildAnalyzerWidget() const
 		{
-			if (!OutputMeter.IsValid() || !OutputOscilloscope.IsValid() || !OutputVectorscope.IsValid())
+			if (!OutputMeter.IsValid() || !OutputOscilloscope.IsValid() || !OutputVectorscope.IsValid() || !OutputSpectrumAnalyzer.IsValid())
 			{
 				return SNullWidget::NullWidget->AsShared();
 			}
@@ -725,6 +726,11 @@ namespace Metasound
 				[
 					OutputVectorscope->GetPanelWidget()
 				]
+				+ SSplitter::Slot()
+				.Value(0.2f)
+				[
+					OutputSpectrumAnalyzer->GetWidget()
+				]
 			];
 		}
 
@@ -1218,6 +1224,11 @@ namespace Metasound
 							OutputVectorscope->CreateVectorscopeWidget(VectorscopePanelLayoutType);
 						}
 
+						if (!OutputSpectrumAnalyzer.IsValid())
+						{
+							OutputSpectrumAnalyzer = MakeShared<AudioWidgets::FAudioSpectrumAnalyzer>(MetaSoundSource->NumChannels, *EditorWorld);
+						}
+
 						return;
 					}
 				}
@@ -1231,6 +1242,7 @@ namespace Metasound
 			OutputMeter.Reset();
 			OutputOscilloscope.Reset();
 			OutputVectorscope.Reset();
+			OutputSpectrumAnalyzer.Reset();
 		}
 
 		void FEditor::ExtendToolbar()
@@ -1622,6 +1634,11 @@ namespace Metasound
 						PreviewComp->SetAudioBusSendPostEffect(AudioBus, 1.0f);
 					}
 
+					if (UAudioBus* AudioBus = OutputSpectrumAnalyzer->GetAudioBus())
+					{
+						PreviewComp->SetAudioBusSendPostEffect(AudioBus, 1.0f);
+					}
+
 					FMetasoundAssetBase* MetasoundAsset = IMetasoundUObjectRegistry::Get().GetObjectAsAssetBase(Metasound);
 					check(MetasoundAsset);
 
diff --git a/Engine/Plugins/Runtime/Metasound/Source/MetasoundEditor/Private/MetasoundEditor.h b/Engine/Plugins/Runtime/Metasound/Source/MetasoundEditor/Private/MetasoundEditor.h
--- a/Engine/Plugins/Runtime/Metasound/Source/MetasoundEditor/Private/MetasoundEditor.h
+++ b/Engine/Plugins/Runtime/Metasound/Source/MetasoundEditor/Private/MetasoundEditor.h
@@ -40,6 +40,7 @@ namespace AudioWidgets
 {
 	class FAudioOscilloscope;
 	class FAudioVectorscope;
+	class FAudioSpectrumAnalyzer;
 }
 
 class FTabManager;
@@ -432,6 +433,9 @@ namespace Metasound
 			/** Vectorscope used in the analyzer tab for visualizing preview output. */
 			TSharedPtr<AudioWidgets::FAudioVectorscope> OutputVectorscope;
 
+			/** Spectrum analyzer used in the analyzer tab for auditioning preview output. */
+			TSharedPtr<AudioWidgets::FAudioSpectrumAnalyzer> OutputSpectrumAnalyzer;
+
 			/** Palette of Node types */
 			TSharedPtr<SMetasoundPalette> Palette;
 
