Fixed missing enum size declaration due to bug in UHT

#rnx

diff --git a/Engine/Plugins/Runtime/AudioWidgets/Source/AudioWidgets/Public/SAudioSpectrumPlot.h b/Engine/Plugins/Runtime/AudioWidgets/Source/AudioWidgets/Public/SAudioSpectrumPlot.h
--- a/Engine/Plugins/Runtime/AudioWidgets/Source/AudioWidgets/Public/SAudioSpectrumPlot.h
+++ b/Engine/Plugins/Runtime/AudioWidgets/Source/AudioWidgets/Public/SAudioSpectrumPlot.h
@@ -8,14 +8,14 @@
 #include "Widgets/SCompoundWidget.h"
 
 UENUM(BlueprintType)
-enum class EAudioSpectrumPlotFrequencyAxisScale
+enum class EAudioSpectrumPlotFrequencyAxisScale : uint8
 {
 	Linear,
 	Logarithmic,
 };
 
 UENUM(BlueprintType)
-enum class EAudioSpectrumPlotFrequencyAxisPixelBucketMode
+enum class EAudioSpectrumPlotFrequencyAxisPixelBucketMode : uint8
 {
 	Sample UMETA(ToolTip = "Plot one data point per frequency axis pixel bucket only, choosing the data point nearest the pixel center."),
 	Peak UMETA(ToolTip = "Plot one data point per frequency axis pixel bucket only, choosing the data point with the highest sound level."),
