Made destroying AudioWidgets::FAudioSpectrumAnalyzer safe in the situation whereby its corresponding Widget does not released at the same time.

diff --git a/Engine/Plugins/Runtime/AudioWidgets/Source/AudioWidgets/Private/AudioSpectrumAnalyzer.cpp b/Engine/Plugins/Runtime/AudioWidgets/Source/AudioWidgets/Private/AudioSpectrumAnalyzer.cpp
--- a/Engine/Plugins/Runtime/AudioWidgets/Source/AudioWidgets/Private/AudioSpectrumAnalyzer.cpp
+++ b/Engine/Plugins/Runtime/AudioWidgets/Source/AudioWidgets/Private/AudioSpectrumAnalyzer.cpp
@@ -19,6 +19,8 @@ namespace AudioWidgets
 	FAudioSpectrumAnalyzer::~FAudioSpectrumAnalyzer()
 	{
 		Teardown();
+
+		Widget->UnbindOnGetAudioSpectrumData();
 	}
 
 	UAudioBus* FAudioSpectrumAnalyzer::GetAudioBus() const
diff --git a/Engine/Plugins/Runtime/AudioWidgets/Source/AudioWidgets/Public/SAudioSpectrumPlot.h b/Engine/Plugins/Runtime/AudioWidgets/Source/AudioWidgets/Public/SAudioSpectrumPlot.h
--- a/Engine/Plugins/Runtime/AudioWidgets/Source/AudioWidgets/Public/SAudioSpectrumPlot.h
+++ b/Engine/Plugins/Runtime/AudioWidgets/Source/AudioWidgets/Public/SAudioSpectrumPlot.h
@@ -154,6 +154,8 @@ public:
 	void SetFrequencyAxisScale(EAudioSpectrumPlotFrequencyAxisScale InFrequencyAxisScale) { FrequencyAxisScale = InFrequencyAxisScale; }
 	void SetFrequencyAxisPixelBucketMode(EAudioSpectrumPlotFrequencyAxisPixelBucketMode InFrequencyAxisPixelBucketMode) { FrequencyAxisPixelBucketMode = InFrequencyAxisPixelBucketMode; }
 
+	void UnbindOnGetAudioSpectrumData() { OnGetAudioSpectrumData.Unbind(); }
+
 private:
 	// Begin SWidget overrides.
 	virtual int32 OnPaint(const FPaintArgs& Args, const FGeometry& AllottedGeometry, const FSlateRect& MyCullingRect, FSlateWindowElementList& OutDrawElements, int32 LayerId, const FWidgetStyle& InWidgetStyle, bool bParentEnabled) const override;
