cache FSlateFontInfo in FixedSampledGridData to avoid possible gc crashes

diff --git a/Engine/Plugins/Runtime/AudioWidgets/Source/AudioWidgets/Private/FixedSampledSequenceGridData.cpp b/Engine/Plugins/Runtime/AudioWidgets/Source/AudioWidgets/Private/FixedSampledSequenceGridData.cpp
--- a/Engine/Plugins/Runtime/AudioWidgets/Source/AudioWidgets/Private/FixedSampledSequenceGridData.cpp
+++ b/Engine/Plugins/Runtime/AudioWidgets/Source/AudioWidgets/Private/FixedSampledSequenceGridData.cpp
@@ -6,19 +6,38 @@
 #include "Framework/Application/SlateApplication.h"
 #include "GenericPlatform/GenericPlatformMath.h"
 
+FFixedSampledSequenceGridData::FFixedSampledSequenceGridData(const uint32 InTotalFrames, const uint32 InSampleRateHz, const FSlateFontInfo& InTicksTimeFont, const float InGridSizePixels)
+	: DisplayRange(TRange<uint32>::Inclusive(0, InTotalFrames))
+	, GridSizePixels(InGridSizePixels)
+	, TicksTimeFont(InTicksTimeFont)
+	, GridFrameRate(InSampleRateHz, 1)
+{
+	UpdateGridMetrics(GridSizePixels);
+}
+
 FFixedSampledSequenceGridData::FFixedSampledSequenceGridData(const uint32 InTotalFrames, const uint32 InSampleRateHz, const float InGridSizePixels /*= 1.f*/, const FSlateFontInfo* InTicksTimeFont /*= nullptr*/)
 	: TotalFrames(InTotalFrames)
 	, DisplayRange(TRange<uint32>::Inclusive(0, InTotalFrames))
 	, GridSizePixels(InGridSizePixels)
-	, TicksTimeFont(InTicksTimeFont)
 	, GridFrameRate(InSampleRateHz, 1)
 {
+	if (InTicksTimeFont)
+	{
+		TicksTimeFont = *InTicksTimeFont;
+	}
+
 	UpdateGridMetrics(GridSizePixels);
 }
 
 void FFixedSampledSequenceGridData::UpdateDisplayRange(const TRange<uint32> InDisplayRange)
 {
 	check(DisplayRange.Size<uint32>() >= 2)
+
+	if (InDisplayRange == DisplayRange)
+	{
+		return;
+	}
+
 	DisplayRange = InDisplayRange;
 	UpdateGridMetrics(GridSizePixels);
 }
@@ -37,13 +56,12 @@ bool FFixedSampledSequenceGridData::UpdateGridMetrics(const float InGridSizePixe
 	double MajorGridStepSeconds = 0.0;
 	float MinTicksPixelSpacing = 30.0f;
 
-	if (TicksTimeFont)
-	{
-		TSharedRef<FSlateFontMeasure> FontMeasureService = FSlateApplication::Get().GetRenderer()->GetFontMeasureService();
-		FString LongestDisplayedTimeString = FString::Printf(TEXT("%.3f"), WaveformDurationSeconds * DisplayRange.GetUpperBoundValue());
-		FVector2D MaxTextSize = FontMeasureService->Measure(LongestDisplayedTimeString, *TicksTimeFont);
-		MinTicksPixelSpacing = MaxTextSize.X;
-	}
+
+	TSharedRef<FSlateFontMeasure> FontMeasureService = FSlateApplication::Get().GetRenderer()->GetFontMeasureService();
+	FString LongestDisplayedTimeString = FString::Printf(TEXT("%.3f"), WaveformDurationSeconds * DisplayRange.GetUpperBoundValue());
+	FVector2D MaxTextSize = FontMeasureService->Measure(LongestDisplayedTimeString, TicksTimeFont);
+	MinTicksPixelSpacing = MaxTextSize.X;
+
 
 
 	if (!GridFrameRate.ComputeGridSpacing(PixelsPerSecond, MajorGridStepSeconds, GridMetrics.NumMinorGridDivisions, MinTicksPixelSpacing + 5.f))
@@ -72,9 +90,18 @@ const FFixedSampledSequenceGridMetrics FFixedSampledSequenceGridData::GetGridMet
 	return GridMetrics;
 }
 
-void FFixedSampledSequenceGridData::SetTicksTimeFont(const FSlateFontInfo* InNewFont)
+void FFixedSampledSequenceGridData::SetTicksTimeFont(const FSlateFontInfo& InNewFont)
 {
 	TicksTimeFont = InNewFont;
+	UpdateGridMetrics(GridSizePixels);	
+}
+
+void FFixedSampledSequenceGridData::SetTicksTimeFont(const FSlateFontInfo * InNewFont)
+{
+	if (InNewFont)
+	{
+		TicksTimeFont = *InNewFont;
+	}
 }
 
 const float FFixedSampledSequenceGridData::SnapPositionToClosestFrame(const float InPixelPosition) const
diff --git a/Engine/Plugins/Runtime/AudioWidgets/Source/AudioWidgets/Private/SAudioOscilloscopePanelWidget.cpp b/Engine/Plugins/Runtime/AudioWidgets/Source/AudioWidgets/Private/SAudioOscilloscopePanelWidget.cpp
--- a/Engine/Plugins/Runtime/AudioWidgets/Source/AudioWidgets/Private/SAudioOscilloscopePanelWidget.cpp
+++ b/Engine/Plugins/Runtime/AudioWidgets/Source/AudioWidgets/Private/SAudioOscilloscopePanelWidget.cpp
@@ -309,7 +309,7 @@ void SAudioOscilloscopePanelWidget::CreateLayout()
 
 void SAudioOscilloscopePanelWidget::CreateGridData(const FFixedSampleSequenceRulerStyle& RulerStyle)
 {
-	SequenceGridData = MakeShared<FFixedSampledSequenceGridData>(DataView.SampleData.Num() / DataView.NumDimensions, DataView.SampleRate, RulerStyle.DesiredWidth, &RulerStyle.TicksTextFont);
+	SequenceGridData = MakeShared<FFixedSampledSequenceGridData>(DataView.SampleData.Num() / DataView.NumDimensions, DataView.SampleRate, RulerStyle.TicksTextFont, RulerStyle.DesiredWidth);
 }
 
 void SAudioOscilloscopePanelWidget::CreateSequenceRuler(TSharedRef<FFixedSampledSequenceGridData> InGridData, const FFixedSampleSequenceRulerStyle& RulerStyle)
diff --git a/Engine/Plugins/Runtime/AudioWidgets/Source/AudioWidgets/Private/SFixedSampledSequenceViewerPanel.cpp b/Engine/Plugins/Runtime/AudioWidgets/Source/AudioWidgets/Private/SFixedSampledSequenceViewerPanel.cpp
--- a/Engine/Plugins/Runtime/AudioWidgets/Source/AudioWidgets/Private/SFixedSampledSequenceViewerPanel.cpp
+++ b/Engine/Plugins/Runtime/AudioWidgets/Source/AudioWidgets/Private/SFixedSampledSequenceViewerPanel.cpp
@@ -147,7 +147,7 @@ void SFixedSampledSequenceViewerPanel::CreateGridData(const FFixedSampleSequence
 {
 	check(RulerStyle)
 
-	SamplesGridData = MakeShared<FFixedSampledSequenceGridData>(DataView.SampleData.Num() / DataView.NumDimensions, DataView.SampleRate, RulerStyle->DesiredWidth, &RulerStyle->TicksTextFont);
+	SamplesGridData = MakeShared<FFixedSampledSequenceGridData>(DataView.SampleData.Num() / DataView.NumDimensions, DataView.SampleRate, RulerStyle->TicksTextFont, RulerStyle->DesiredWidth);
 }
 
 void SFixedSampledSequenceViewerPanel::ReceiveSequenceView(const FFixedSampledSequenceView InData, const uint32 FirstSampleIndex)
diff --git a/Engine/Plugins/Runtime/AudioWidgets/Source/AudioWidgets/Public/FixedSampledSequenceGridData.h b/Engine/Plugins/Runtime/AudioWidgets/Source/AudioWidgets/Public/FixedSampledSequenceGridData.h
--- a/Engine/Plugins/Runtime/AudioWidgets/Source/AudioWidgets/Public/FixedSampledSequenceGridData.h
+++ b/Engine/Plugins/Runtime/AudioWidgets/Source/AudioWidgets/Public/FixedSampledSequenceGridData.h
@@ -5,22 +5,30 @@
 #include "Delegates/Delegate.h"
 #include "HAL/Platform.h"
 #include "IFixedSampledSequenceGridService.h"
+#include "Fonts/SlateFontInfo.h"
 #include "Math/Range.h"
 #include "Misc/FrameRate.h"
 
 DECLARE_MULTICAST_DELEGATE(FOnWaveformEditorGridUpdated)
 
-struct FSlateFontInfo;
-
 class AUDIOWIDGETS_API FFixedSampledSequenceGridData : public IFixedSampledSequenceGridService
 {
 public:
+	explicit FFixedSampledSequenceGridData(const uint32 InTotalFrames, const uint32 InSampleRateHz, const FSlateFontInfo& InTicksTimeFont, const float InGridSizePixels = 1.f);
+
+	UE_DEPRECATED(5.4, "Constructor taking a const FSlateFontInfo* is deprecated, use the one taking a refrence instead")
 	explicit FFixedSampledSequenceGridData(const uint32 InTotalFrames, const uint32 InSampleRateHz, const float InGridSizePixels = 1.f, const FSlateFontInfo* InTicksTimeFont = nullptr);
 
+
 	void UpdateDisplayRange(const TRange<uint32> InDisplayRange);
 	bool UpdateGridMetrics(const float InGridSizePixels);
 	virtual const FFixedSampledSequenceGridMetrics GetGridMetrics() const override;
+	void SetTicksTimeFont(const FSlateFontInfo& InNewFont);
+
+	UE_DEPRECATED(5.4, "SetTicksTimeFont taking a const FSlateFontInfo* is deprecated, use the one taking a refrence instead")
 	void SetTicksTimeFont(const FSlateFontInfo* InNewFont);
+
+	
 	const float SnapPositionToClosestFrame(const float InPixelPosition) const;
 
 	FOnWaveformEditorGridUpdated OnGridMetricsUpdated;
@@ -31,6 +39,6 @@ private:
 	TRange<uint32> DisplayRange;
 
 	float GridSizePixels = 0.f;
-	const FSlateFontInfo* TicksTimeFont = nullptr;
+	FSlateFontInfo TicksTimeFont;
 	FFrameRate GridFrameRate;
 };
\ No newline at end of file
