Added SAudioSpectrogram slate widget.

diff --git a/Engine/Plugins/Runtime/AudioWidgets/Source/AudioWidgets/Private/AudioColorMapper.cpp b/Engine/Plugins/Runtime/AudioWidgets/Source/AudioWidgets/Private/AudioColorMapper.cpp
new file mode 100644
--- /dev/null
+++ b/Engine/Plugins/Runtime/AudioWidgets/Source/AudioWidgets/Private/AudioColorMapper.cpp
@@ -0,0 +1,30 @@
+// Copyright Epic Games, Inc. All Rights Reserved.
+
+#include "AudioColorMapper.h"
+
+FColor FAudioColorMapper::GetColorFromValue(const float Value) const
+{
+	const float ScaledValue = ValueScaling * (Value - MinValue);
+	const float ScaledClampedValue = FMath::Clamp(ScaledValue, 0.0f, 1.0f);
+	switch (ColorMap)
+	{
+	default:
+	case EAudioColorMap::BlackToWhite:
+		return ColorMapBlackToWhite(ScaledClampedValue);
+	case EAudioColorMap::WhiteToBlack:
+		return ColorMapWhiteToBlack(ScaledClampedValue);
+	}
+}
+
+FColor FAudioColorMapper::ColorMapBlackToWhite(const float T)
+{
+	const uint8 Quantized = FColor::QuantizeUNormFloatTo8(T);
+	return FColor(Quantized, Quantized, Quantized);
+}
+
+FColor FAudioColorMapper::ColorMapWhiteToBlack(const float T)
+{
+	const uint8 Quantized = FColor::QuantizeUNormFloatTo8(1.0f - T);
+	return FColor(Quantized, Quantized, Quantized);
+}
+
diff --git a/Engine/Plugins/Runtime/AudioWidgets/Source/AudioWidgets/Private/AudioSpectrogramViewport.cpp b/Engine/Plugins/Runtime/AudioWidgets/Source/AudioWidgets/Private/AudioSpectrogramViewport.cpp
new file mode 100644
--- /dev/null
+++ b/Engine/Plugins/Runtime/AudioWidgets/Source/AudioWidgets/Private/AudioSpectrogramViewport.cpp
@@ -0,0 +1,427 @@
+// Copyright Epic Games, Inc. All Rights Reserved.
+
+#include "AudioSpectrogramViewport.h"
+
+#include "Framework/Application/SlateApplication.h"
+#include "Rendering/SlateRenderer.h"
+#include "Slate/SlateTextures.h"
+
+////////////////////////////////////////////////////////////////////////////////////////////////////////////////////////////////////////////////////////
+
+/**
+ * Helper class for getting the normalized read position in the source spectrum data for a given target pixel to be rendered.
+ */
+class FAudioSpectrogramRenderTransformHelper
+{
+public:
+	FAudioSpectrogramRenderTransformHelper(const FAudioSpectrogramViewportRenderParams& RenderParams, const bool bSpectrumDataHasLogSpacedFreqencies, const float SpectrumDataMinFrequency, const float SpectrumDataMaxFrequency)
+		: PixelToFreqTransform(RenderParams.FrequencyAxisScale, RenderParams.NumPixelsPerRow, RenderParams.ViewMinFrequency, RenderParams.ViewMaxFrequency)
+		, FreqToDataPosTransform(bSpectrumDataHasLogSpacedFreqencies, SpectrumDataMinFrequency, SpectrumDataMaxFrequency)
+	{
+		//
+	}
+
+	// Given a pixel index, get the normalized position to sample the data at.
+	float PixelToNormalizedDataPos(const float Pixel) const
+	{
+		const float Frequency = PixelToFreqTransform.PixelToFrequency(Pixel);
+		return FreqToDataPosTransform.FrequencyToNormalizedDataPos(Frequency);
+	}
+
+	// Given a pixel index, get the ideal mip size that would best resolve the data for the position the data would be sampled at.
+	float GetDesiredMipSize(const int Pixel) const
+	{
+		const float NormalizedDataPosLo = PixelToNormalizedDataPos(float(Pixel) - 0.5f);
+		const float NormalizedDataPosHi = PixelToNormalizedDataPos(float(Pixel) + 0.5f);
+		const float NormalizedDataPosDelta = FMath::Abs(NormalizedDataPosHi - NormalizedDataPosLo);
+		constexpr float MinDelta = 1.0f / float(1 << 24);
+		return 1.0f / FMath::Max(NormalizedDataPosDelta, MinDelta);
+	}
+
+private:
+	/**
+	 * Internal helper class for getting the frequency for a given target pixel to be rendered.
+	 */
+	class FPixelToFreqTransform
+	{
+	public:
+		FPixelToFreqTransform(const EAudioSpectrogramFrequencyAxisScale InFrequencyAxisScale, const int NumPixelsPerRow, const float InViewMinFrequency, const float InViewMaxFrequency)
+			: FrequencyAxisScale(InFrequencyAxisScale)
+			, TransformedViewMinFrequency(ForwardTransformFrequency(InViewMinFrequency))
+			, ScalingValue((ForwardTransformFrequency(InViewMaxFrequency) - TransformedViewMinFrequency) / float(NumPixelsPerRow - 1))
+		{
+			check(NumPixelsPerRow >= 2);
+		}
+
+		float PixelToFrequency(const float Index) const
+		{
+			const float TransformedFrequency = TransformedViewMinFrequency + (Index * ScalingValue);
+			return InverseTransformFrequency(TransformedFrequency);
+		}
+
+	private:
+		float ForwardTransformFrequency(const float Frequency) const
+		{
+			return (FrequencyAxisScale == EAudioSpectrogramFrequencyAxisScale::Logarithmic) ? FMath::Loge(Frequency) : Frequency;
+		}
+
+		float InverseTransformFrequency(const float TransformedFrequency) const
+		{
+			return (FrequencyAxisScale == EAudioSpectrogramFrequencyAxisScale::Logarithmic) ? FMath::Exp(TransformedFrequency) : TransformedFrequency;
+		}
+
+		const EAudioSpectrogramFrequencyAxisScale FrequencyAxisScale;
+		const float TransformedViewMinFrequency;
+		const float ScalingValue;
+	};
+
+	/**
+	 * Internal helper class for getting the normalized read position in the source spectrum data for a given frequency.
+	 */
+	class FFreqToDataPosTransform
+	{
+	public:
+		FFreqToDataPosTransform(const bool bInSpectrumDataHasLogSpacedFreqencies, const float InSpectrumDataMinFrequency, const float InSpectrumDataMaxFrequency)
+			: bSpectrumDataHasLogSpacedFreqencies(bInSpectrumDataHasLogSpacedFreqencies)
+			, TransformedSpectrumDataMinFrequency(TransformFrequency(InSpectrumDataMinFrequency))
+			, ScalingValue(1.0f / (TransformFrequency(InSpectrumDataMaxFrequency) - TransformedSpectrumDataMinFrequency))
+		{
+			check(InSpectrumDataMaxFrequency != InSpectrumDataMinFrequency);
+		}
+
+		float FrequencyToNormalizedDataPos(const float Frequency) const
+		{
+			return ScalingValue * (TransformFrequency(Frequency) - TransformedSpectrumDataMinFrequency);
+		}
+
+	private:
+		float TransformFrequency(const float Frequency) const
+		{
+			return (bSpectrumDataHasLogSpacedFreqencies) ? FMath::Loge(Frequency) : Frequency;
+		}
+
+		const bool bSpectrumDataHasLogSpacedFreqencies;
+		const float TransformedSpectrumDataMinFrequency;
+		const float ScalingValue;
+	};
+
+	const FPixelToFreqTransform PixelToFreqTransform;
+	const FFreqToDataPosTransform FreqToDataPosTransform;
+};
+
+////////////////////////////////////////////////////////////////////////////////////////////////////////////////////////////////////////////////////////
+
+FAudioSpectrogramViewport::FAudioSpectrogramViewport()
+{
+	// Max history size should default to zero (we shall only start collecting spectrogram frames once render params have been set):
+	ensure(RenderParams.NumRows == 0);
+}
+
+FAudioSpectrogramViewport::~FAudioSpectrogramViewport()
+{
+	// Release texture if we have one:
+	if (UpdatableTexture != nullptr)
+	{
+		if (FSlateRenderer* Renderer = FSlateApplication::Get().GetRenderer())
+		{
+			Renderer->ReleaseUpdatableTexture(UpdatableTexture);
+		}
+
+		UpdatableTexture = nullptr;
+		UpdatableTexturePixelFormat = PF_Unknown;
+	}
+}
+
+FIntPoint FAudioSpectrogramViewport::GetSize() const
+{
+	const FSlateShaderResource* SlateResource = GetViewportRenderTargetTexture();
+	if (SlateResource != nullptr)
+	{
+		return FIntPoint(SlateResource->GetWidth(), SlateResource->GetHeight());
+	}
+
+	return FIntPoint::ZeroValue;
+}
+
+void FAudioSpectrogramViewport::SetRenderParams(const FAudioSpectrogramViewportRenderParams& InRenderParams)
+{
+	if (RenderParams.NumRows != InRenderParams.NumRows ||
+		RenderParams.NumPixelsPerRow != InRenderParams.NumPixelsPerRow ||
+		RenderParams.ViewMinFrequency != InRenderParams.ViewMinFrequency ||
+		RenderParams.ViewMaxFrequency != InRenderParams.ViewMaxFrequency ||
+		RenderParams.ColorMapMinSoundLevel != InRenderParams.ColorMapMinSoundLevel ||
+		RenderParams.ColorMapMaxSoundLevel != InRenderParams.ColorMapMaxSoundLevel ||
+		RenderParams.ColorMap != InRenderParams.ColorMap ||
+		RenderParams.FrequencyAxisScale != InRenderParams.FrequencyAxisScale ||
+		RenderParams.FrequencyAxisPixelBucketMode != InRenderParams.FrequencyAxisPixelBucketMode)
+	{
+		RenderParams = InRenderParams;
+
+		if (History.Num() > RenderParams.NumRows)
+		{
+			// Reduce size of history by removing frames from the start:
+			const int ExcessFrames = History.Num() - RenderParams.NumRows;
+			History.PopFront(ExcessFrames);
+		}
+
+		// Render params changed, invalidate all cached pixel data:
+		UpdateTextureData(true);
+	}
+}
+
+void FAudioSpectrogramViewport::AddFrame(const FAudioSpectrogramFrameData& SpectrogramFrameData)
+{
+	if (RenderParams.NumRows == 0)
+	{
+		return;
+	}
+
+	if (History.Num() == RenderParams.NumRows)
+	{
+		// History is full. Make space for new frame by removing oldest from the start of the history: 
+		History.PopFront();
+	}
+
+	// Add the new frame to the end of the history buffer:
+	History.Emplace(SpectrogramFrameData, RenderParams.FrequencyAxisPixelBucketMode);
+
+	// Spectrogram data changed, update texture data:
+	UpdateTextureData();
+}
+
+void FAudioSpectrogramViewport::UpdateTextureData(const bool bInvalidateCachedPixelData)
+{
+	TRACE_CPUPROFILER_EVENT_SCOPE(FAudioSpectrogramViewport::UpdateTextureData);
+
+	// Figure out required texture size:
+	const int32 Width = RenderParams.NumPixelsPerRow;
+	const int32 Height = RenderParams.NumRows;
+
+	if (UpdatableTexture == nullptr)
+	{
+		// Try to create a texture of the required size:
+		if (FSlateRenderer* Renderer = FSlateApplication::Get().GetRenderer())
+		{
+			UpdatableTexture = Renderer->CreateUpdatableTexture(Width, Height);
+			UpdatableTexturePixelFormat = Renderer->GetSlateRecommendedColorFormat();
+		}
+	}
+
+	if (UpdatableTexture != nullptr)
+	{
+		// We shall write rows of pixels to RawData:
+		TArray<uint8> RawData;
+
+		// Set RawData to the size requried for the texture:
+		check(UpdatableTexturePixelFormat != PF_Unknown);
+		const int32 BytesPerPixel = GPixelFormats[UpdatableTexturePixelFormat].BlockBytes;
+		const int32 BytesPerRow = Width * BytesPerPixel;
+		const int32 SizeInBytes = Height * BytesPerRow;
+		RawData.AddZeroed(SizeInBytes);
+
+		// Ensure the texture is of the expected format:
+		if (ensure(BytesPerPixel == sizeof(FColor) && UpdatableTexturePixelFormat == PF_B8G8R8A8))
+		{
+			// Go over each frame in the History buffer:
+			check(History.Num() <= Height);
+			const int32 NumSkipRows = Height - History.Num(); // If History size is smaller than texture height, skip over rows in the texture to fill right to the bottom.
+			for (int32 RowIndex = NumSkipRows; RowIndex < Height; RowIndex++)
+			{
+				FHistoryFrameData& HistoryFrameData = History[RowIndex - NumSkipRows];
+
+				// Generate/regenerate pixel data for this frame if required:
+				if (HistoryFrameData.GetCachedPixels().Num() != Width || bInvalidateCachedPixelData)
+				{
+					HistoryFrameData.UpdateCachedPixels(RenderParams);
+				}
+
+				// Copy generated pixel data to RawData:
+				const TConstArrayView<FColor> CachedPixels = HistoryFrameData.GetCachedPixels();
+				check(CachedPixels.Num() == Width);
+				memcpy(&RawData[RowIndex * BytesPerRow], CachedPixels.GetData(), BytesPerRow);
+			}
+		}
+
+		// Apply RawData to the texture:
+		FSlateTextureData* BulkData = new FSlateTextureData(Width, Height, BytesPerPixel, MoveTemp(RawData));
+		UpdatableTexture->UpdateTextureThreadSafeWithTextureData(BulkData); // Ownership of BulkData is transferred to the render system here.
+	}
+}
+
+////////////////////////////////////////////////////////////////////////////////////////////////////////////////////////////////////////////////////////
+
+FAudioSpectrogramViewport::FHistoryFrameData::FHistoryFrameData(const FAudioSpectrogramFrameData& SpectrogramFrameData, const EAudioSpectrogramFrequencyAxisPixelBucketMode InFrequencyAxisPixelBucketMode)
+	: NumSoundLevels(SpectrogramFrameData.SquaredMagnitudes.Num())
+	, SpectrumDataMinFrequency(SpectrogramFrameData.MinFrequency)
+	, SpectrumDataMaxFrequency(SpectrogramFrameData.MaxFrequency)
+	, bSpectrumDataHasLogSpacedFreqencies(SpectrogramFrameData.bLogSpacedFreqencies)
+	, SoundLevelsMipChainPixelBucketMode(InFrequencyAxisPixelBucketMode)
+{
+	// Allocate required space for mip chain:
+	SetMipChainLengthUninitialized(InFrequencyAxisPixelBucketMode);
+
+	// Convert squared magnitudes to dB, and write to mip zero:
+	const float ClampMinMagnitudeSquared = FMath::Pow(10.0f, -200.0f / 10.0f); // Clamp at -200dB
+	for (int Index = 0; Index < NumSoundLevels; Index++)
+	{
+		const float MagnitudeSquared = FMath::Max(SpectrogramFrameData.SquaredMagnitudes[Index], ClampMinMagnitudeSquared);
+		SoundLevelsMipChain[Index] = 10.0f * FMath::LogX(10.0f, MagnitudeSquared);
+	}
+
+	// Generate remaining mips if required:
+	UpdateMipChain(InFrequencyAxisPixelBucketMode);
+}
+
+FAudioSpectrogramViewport::FHistoryFrameData::FHistoryFrameData(FHistoryFrameData&& Other)
+	: NumSoundLevels(Other.NumSoundLevels)
+	, SpectrumDataMinFrequency(Other.SpectrumDataMinFrequency)
+	, SpectrumDataMaxFrequency(Other.SpectrumDataMaxFrequency)
+	, bSpectrumDataHasLogSpacedFreqencies(Other.bSpectrumDataHasLogSpacedFreqencies)
+	, SoundLevelsMipChain(MoveTemp(Other.SoundLevelsMipChain))
+	, SoundLevelsMipChainPixelBucketMode(Other.SoundLevelsMipChainPixelBucketMode)
+	, CachedPixels(MoveTemp(Other.CachedPixels))
+{
+	//
+}
+
+void FAudioSpectrogramViewport::FHistoryFrameData::UpdateCachedPixels(const FAudioSpectrogramViewportRenderParams& InRenderParams)
+{
+	// Create local helpers:
+	const FAudioSpectrogramRenderTransformHelper RenderTransformHelper(InRenderParams, bSpectrumDataHasLogSpacedFreqencies, SpectrumDataMinFrequency, SpectrumDataMaxFrequency);
+	const FAudioColorMapper ColorMapper(InRenderParams.ColorMapMinSoundLevel, InRenderParams.ColorMapMaxSoundLevel, InRenderParams.ColorMap);
+
+	// Regenerate mip chain if FrequencyAxisPixelBucketMode has changed:
+	if (SoundLevelsMipChainPixelBucketMode != InRenderParams.FrequencyAxisPixelBucketMode)
+	{
+		UpdateMipChain(InRenderParams.FrequencyAxisPixelBucketMode);
+		SoundLevelsMipChainPixelBucketMode = InRenderParams.FrequencyAxisPixelBucketMode;
+	}
+
+	// Ensure pixel data array is of the required size before writing to it:
+	CachedPixels.SetNumUninitialized(InRenderParams.NumPixelsPerRow);
+
+	if (HasMips())
+	{
+		for (int Index = 0; Index < InRenderParams.NumPixelsPerRow; Index++)
+		{
+			// Calculate mip LOD:
+			const float DesiredMipSize = RenderTransformHelper.GetDesiredMipSize(Index);
+
+			// Sample from mips with interpolation:
+			const float NormalizedDataPos = RenderTransformHelper.PixelToNormalizedDataPos(Index);
+			const float InterpolatedSoundLevel = GetMipInterpolatedSoundLevel(DesiredMipSize, NormalizedDataPos);
+
+			// Apply colormap and write pixel:
+			CachedPixels[Index] = ColorMapper.GetColorFromValue(InterpolatedSoundLevel);
+		}
+	}
+	else
+	{
+		// Render using mip zero only (which is all we have):
+		const TConstArrayView<float> Mip0(SoundLevelsMipChain);
+		for (int Index = 0; Index < InRenderParams.NumPixelsPerRow; Index++)
+		{
+			// Sample data with interpolation:
+			const float NormalizedDataPos = RenderTransformHelper.PixelToNormalizedDataPos(Index);
+			const float InterpolatedSoundLevel = GetInterpolatedSoundLevel(Mip0, NormalizedDataPos);
+
+			// Apply colormap and write pixel:
+			CachedPixels[Index] = ColorMapper.GetColorFromValue(InterpolatedSoundLevel);
+		}
+	}
+}
+
+void FAudioSpectrogramViewport::FHistoryFrameData::SetMipChainLengthUninitialized(const EAudioSpectrogramFrequencyAxisPixelBucketMode InFrequencyAxisPixelBucketMode)
+{
+	// Only supporting mip maps for power of two sizes:
+	const bool bUseMipMaps = FMath::IsPowerOfTwo(NumSoundLevels) && (InFrequencyAxisPixelBucketMode != EAudioSpectrogramFrequencyAxisPixelBucketMode::Sample) && !bSpectrumDataHasLogSpacedFreqencies;
+
+	const int SoundLevelsMipChainLength = (bUseMipMaps) ? 2 * NumSoundLevels - 1 : NumSoundLevels;
+	SoundLevelsMipChain.SetNumUninitialized(SoundLevelsMipChainLength);
+}
+
+void FAudioSpectrogramViewport::FHistoryFrameData::UpdateMipChain(const EAudioSpectrogramFrequencyAxisPixelBucketMode InFrequencyAxisPixelBucketMode)
+{
+	// Mip zero should already exist:
+	check(SoundLevelsMipChain.Num() >= NumSoundLevels);
+
+	// Resize SoundLevelsMipChain if necessary for the current pixel bucket mode:
+	SetMipChainLengthUninitialized(InFrequencyAxisPixelBucketMode);
+
+	// Mip zero should still exist:
+	check(SoundLevelsMipChain.Num() >= NumSoundLevels);
+
+	if (HasMips())
+	{
+		// Only these two modes use mip maps:
+		check(InFrequencyAxisPixelBucketMode == EAudioSpectrogramFrequencyAxisPixelBucketMode::Peak || InFrequencyAxisPixelBucketMode == EAudioSpectrogramFrequencyAxisPixelBucketMode::Average);
+
+		// Start reading from start of mip zero:
+		int ReadIndex = 0;
+
+		// Start writing data directly after mip zero:
+		int WriteIndex = NumSoundLevels;
+
+		// Generate mips:
+		for (int CurrentMipSize = NumSoundLevels / 2; CurrentMipSize >= 1; CurrentMipSize /= 2)
+		{
+			for (int Count = 0; Count < CurrentMipSize; Count++)
+			{
+				const float ReadA = SoundLevelsMipChain[ReadIndex++];
+				const float ReadB = SoundLevelsMipChain[ReadIndex++];
+				if (InFrequencyAxisPixelBucketMode == EAudioSpectrogramFrequencyAxisPixelBucketMode::Peak)
+				{
+					SoundLevelsMipChain[WriteIndex++] = FMath::Max(ReadA, ReadB);
+				}
+				else// if (InFrequencyAxisPixelBucketMode == EAudioSpectrogramFrequencyAxisPixelBucketMode::Average)
+				{
+					SoundLevelsMipChain[WriteIndex++] = 0.5f * (ReadA + ReadB);
+				}
+			}
+		}
+
+		// We have written the full mip chain:
+		check(WriteIndex == SoundLevelsMipChain.Num());
+	}
+}
+
+TConstArrayView<float> FAudioSpectrogramViewport::FHistoryFrameData::GetMipFromMipSize(const int MipSize) const
+{
+	check(FMath::IsPowerOfTwo(MipSize));
+	check(FMath::IsPowerOfTwo(NumSoundLevels));
+	check(MipSize <= NumSoundLevels);
+
+	const int StartIndex = 2 * (NumSoundLevels - MipSize);
+	return TConstArrayView<float>(SoundLevelsMipChain).Slice(StartIndex, MipSize);
+}
+
+float FAudioSpectrogramViewport::FHistoryFrameData::GetMipInterpolatedSoundLevel(const float DesiredMipSize, const float NormalizedDataPos) const
+{
+	// Get mips for mip LOD:
+	const int MipHiSize = FMath::Clamp(FMath::RoundUpToPowerOfTwo(FMath::CeilToInt(DesiredMipSize)), 2, NumSoundLevels);
+	const int MipLoSize = MipHiSize / 2;
+	const TConstArrayView<float> MipHi = GetMipFromMipSize(MipHiSize);
+	const TConstArrayView<float> MipLo = GetMipFromMipSize(MipLoSize);
+
+	// Sample from mips with interpolation:
+	const float MipHiSample = GetInterpolatedSoundLevel(MipHi, NormalizedDataPos);
+	const float MipLoSample = GetInterpolatedSoundLevel(MipLo, NormalizedDataPos);
+
+	// Interpolate between mips:
+	const float MipLerpParam = FMath::Clamp(float(DesiredMipSize - MipLoSize) / float(MipHiSize - MipLoSize), 0.0f, 1.0f);
+	return FMath::Lerp(MipLoSample, MipHiSample, MipLerpParam);
+}
+
+float FAudioSpectrogramViewport::FHistoryFrameData::GetInterpolatedSoundLevel(const TConstArrayView<float> MipData, const float NormalizedDataPos)
+{
+	check(NormalizedDataPos >= 0.0f);
+
+	const int MipSize = MipData.Num();
+	const float IndexFloat = NormalizedDataPos * (MipSize - 1);
+	const float IndexFloor = FMath::Floor(IndexFloat);
+	const float LerpParam = IndexFloat - IndexFloor;
+	const int IndexLo = FMath::Min(static_cast<int>(IndexFloor) + 0, MipSize - 1);
+	const int IndexHi = FMath::Min(static_cast<int>(IndexFloor) + 1, MipSize - 1);
+	return FMath::Lerp(MipData[IndexLo], MipData[IndexHi], LerpParam);
+}
+
diff --git a/Engine/Plugins/Runtime/AudioWidgets/Source/AudioWidgets/Private/SAudioSpectrogram.cpp b/Engine/Plugins/Runtime/AudioWidgets/Source/AudioWidgets/Private/SAudioSpectrogram.cpp
new file mode 100644
--- /dev/null
+++ b/Engine/Plugins/Runtime/AudioWidgets/Source/AudioWidgets/Private/SAudioSpectrogram.cpp
@@ -0,0 +1,63 @@
+// Copyright Epic Games, Inc. All Rights Reserved.
+
+#include "SAudioSpectrogram.h"
+
+void SAudioSpectrogram::Construct(const FArguments& InArgs)
+{
+	ViewMinFrequency = InArgs._ViewMinFrequency;
+	ViewMaxFrequency = InArgs._ViewMaxFrequency;
+	ColorMapMinSoundLevel = InArgs._ColorMapMinSoundLevel;
+	ColorMapMaxSoundLevel = InArgs._ColorMapMaxSoundLevel;
+	ColorMap = InArgs._ColorMap;
+	FrequencyAxisScale = InArgs._FrequencyAxisScale;
+	FrequencyAxisPixelBucketMode = InArgs._FrequencyAxisPixelBucketMode;
+	Orientation = InArgs._Orientation;
+
+	SpectrogramViewport = MakeShareable(new FAudioSpectrogramViewport());
+}
+
+void SAudioSpectrogram::AddFrame(const FAudioSpectrogramFrameData& SpectrogramFrameData)
+{
+	SpectrogramViewport->AddFrame(SpectrogramFrameData);
+}
+
+int32 SAudioSpectrogram::OnPaint(const FPaintArgs& Args, const FGeometry& AllottedGeometry, const FSlateRect& MyCullingRect, FSlateWindowElementList& OutDrawElements, int32 LayerId, const FWidgetStyle& InWidgetStyle, bool bParentEnabled) const
+{
+	// Get post DPI-scaled sizing:
+	const FVector2f AbsoluteSize = AllottedGeometry.GetAbsoluteSize();
+
+	// Round down size to whole pixels:
+	const int RenderWidth = FMath::FloorToInt(AbsoluteSize.X);
+	const int RenderHeight = FMath::FloorToInt(AbsoluteSize.Y);
+	const FVector2f RenderSize(RenderWidth, RenderHeight);
+	const FVector2f AdjustedLocalSize = TransformVector(Inverse(AllottedGeometry.GetAccumulatedRenderTransform()), RenderSize);
+
+	// Apply rotation to viewport as required by Orientation attribute:
+	const bool bRotateViewport = (Orientation.Get() == EOrientation::Orient_Horizontal);
+	const FMatrix2x2f ScaleAndRotate(0.0f, -RenderSize.Y / RenderSize.X, RenderSize.X / RenderSize.Y, 0.0f); // Scale and rotate to swap X and Y.
+	const FSlateRenderTransform ViewportRenderTransform = (bRotateViewport) ? FSlateRenderTransform(ScaleAndRotate) : FSlateRenderTransform();
+	const FVector2f ViewportRenderTransformPivot(0.5f);
+
+	// Create a child geometry using adjusted size and possible render rotation:
+	const FGeometry ChildGeometry = AllottedGeometry.MakeChild(AdjustedLocalSize, FSlateLayoutTransform(), ViewportRenderTransform, ViewportRenderTransformPivot);
+	
+	// Create spectrogram render params, setting viewport history size to exact pixel size:
+	const FAudioSpectrogramViewportRenderParams RenderParams
+	{
+		.NumRows = (bRotateViewport) ? RenderWidth : RenderHeight,
+		.NumPixelsPerRow = FMath::Max((bRotateViewport) ? RenderHeight : RenderWidth, 2),
+		.ViewMinFrequency = ViewMinFrequency.Get(),
+		.ViewMaxFrequency = ViewMaxFrequency.Get(),
+		.ColorMapMinSoundLevel = ColorMapMinSoundLevel.Get(),
+		.ColorMapMaxSoundLevel = ColorMapMaxSoundLevel.Get(),
+		.ColorMap = ColorMap.Get(),
+		.FrequencyAxisScale = FrequencyAxisScale.Get(),
+		.FrequencyAxisPixelBucketMode = FrequencyAxisPixelBucketMode.Get(),
+	};
+	SpectrogramViewport->SetRenderParams(RenderParams);
+
+	// Create the viewport using child geometry:
+	FSlateDrawElement::MakeViewport(OutDrawElements, LayerId, ChildGeometry.ToPaintGeometry(), SpectrogramViewport, ESlateDrawEffect::NoGamma);
+
+	return LayerId + 1;
+}
diff --git a/Engine/Plugins/Runtime/AudioWidgets/Source/AudioWidgets/Public/AudioColorMapper.h b/Engine/Plugins/Runtime/AudioWidgets/Source/AudioWidgets/Public/AudioColorMapper.h
new file mode 100644
--- /dev/null
+++ b/Engine/Plugins/Runtime/AudioWidgets/Source/AudioWidgets/Public/AudioColorMapper.h
@@ -0,0 +1,38 @@
+// Copyright Epic Games, Inc. All Rights Reserved.
+
+#pragma once
+
+#include "Math/Color.h"
+#include "UObject/ObjectMacros.h"
+
+UENUM(BlueprintType)
+enum class EAudioColorMap : uint8
+{
+	BlackToWhite,
+	WhiteToBlack,
+};
+
+/**
+ * Helper class for normalizing and applying a colormapping to unnormalized data.
+ */
+class FAudioColorMapper
+{
+public:
+	FAudioColorMapper(const float InMinValue, const float InMaxValue, const EAudioColorMap InColorMap)
+		: ValueScaling(1.0f / (InMaxValue - InMinValue))
+		, MinValue(InMinValue)
+		, ColorMap(InColorMap)
+	{
+		//
+	}
+
+	FColor GetColorFromValue(const float Value) const;
+
+	static FColor ColorMapBlackToWhite(const float T);
+	static FColor ColorMapWhiteToBlack(const float T);
+	
+private:
+	const float ValueScaling;
+	const float MinValue;
+	const EAudioColorMap ColorMap;
+};
diff --git a/Engine/Plugins/Runtime/AudioWidgets/Source/AudioWidgets/Public/AudioSpectrogramViewport.h b/Engine/Plugins/Runtime/AudioWidgets/Source/AudioWidgets/Public/AudioSpectrogramViewport.h
new file mode 100644
--- /dev/null
+++ b/Engine/Plugins/Runtime/AudioWidgets/Source/AudioWidgets/Public/AudioSpectrogramViewport.h
@@ -0,0 +1,134 @@
+// Copyright Epic Games, Inc. All Rights Reserved.
+
+#pragma once
+
+#include "AudioColorMapper.h"
+#include "Containers/RingBuffer.h"
+#include "PixelFormat.h"
+#include "Rendering/RenderingCommon.h"
+#include "Textures/SlateUpdatableTexture.h"
+
+UENUM(BlueprintType)
+enum class EAudioSpectrogramFrequencyAxisScale : uint8
+{
+	Linear,
+	Logarithmic,
+};
+
+UENUM(BlueprintType)
+enum class EAudioSpectrogramFrequencyAxisPixelBucketMode : uint8
+{
+	Sample UMETA(ToolTip = "Plot one data point per frequency axis pixel bucket only, choosing the data point nearest the pixel center."),
+	Peak UMETA(ToolTip = "Plot one data point per frequency axis pixel bucket only, choosing the data point with the highest sound level."),
+	Average UMETA(ToolTip = "Plot the average of the data points in each frequency axis pixel bucket."),
+};
+
+/**
+ * Parameters for the spectrogram viewport.
+ */
+struct FAudioSpectrogramViewportRenderParams
+{
+	int NumRows = 0;
+	int NumPixelsPerRow = 0;
+	float ViewMinFrequency = 0.0f;
+	float ViewMaxFrequency = 0.0f;
+	float ColorMapMinSoundLevel = 0.0f;
+	float ColorMapMaxSoundLevel = 0.0f;
+	EAudioColorMap ColorMap = EAudioColorMap::BlackToWhite;
+	EAudioSpectrogramFrequencyAxisScale FrequencyAxisScale = EAudioSpectrogramFrequencyAxisScale::Logarithmic;
+	EAudioSpectrogramFrequencyAxisPixelBucketMode FrequencyAxisPixelBucketMode = EAudioSpectrogramFrequencyAxisPixelBucketMode::Average;
+};
+
+/**
+ * Struct for passing spectrum data for one frame of the spectrogram
+ */
+struct FAudioSpectrogramFrameData
+{
+	TConstArrayView<float> SquaredMagnitudes;
+	float MinFrequency = 0.0f;
+	float MaxFrequency = 0.0f;
+	bool bLogSpacedFreqencies = false;
+};
+
+/**
+ * Maintains a history of spectrogram frames, and presents their visualization by implementing ISlateViewport.
+ */
+class FAudioSpectrogramViewport : public ISlateViewport, public TSharedFromThis<FAudioSpectrogramViewport>
+{
+public:
+	FAudioSpectrogramViewport();
+	~FAudioSpectrogramViewport();
+
+	/* ISlateViewport interface */
+	virtual FIntPoint GetSize() const override;
+	virtual class FSlateShaderResource* GetViewportRenderTargetTexture() const override { return (UpdatableTexture) ? UpdatableTexture->GetSlateResource() : nullptr; }
+	virtual bool RequiresVsync() const override { return false; }
+	virtual bool AllowScaling() const override { return true; }
+
+	// The render params must be set before the FAudioSpectrogramViewport will do anything.
+	void SetRenderParams(const FAudioSpectrogramViewportRenderParams& InRenderParams);
+
+	// Add the given spectrum frame data to the history buffer.
+	void AddFrame(const FAudioSpectrogramFrameData& SpectrogramFrameData);
+
+private:
+	// If row data or render params have changed, we must update the texture data. If render params change, we must also invalidate any previously cached pixel data for all rows.
+	void UpdateTextureData(const bool bInvalidateCachedPixelData = false);
+
+	/**
+	 * Stores one frame of spectrogram data. Also generates the visual representation of this data as colored pixels, according to the given render params.
+	 */
+	class FHistoryFrameData : public FNoncopyable
+	{
+	public:
+		// Construct a history row to store the given spectrum data. Optionally also specify the pixel bucket mode to pre-allocate space for mips if required.
+		FHistoryFrameData(const FAudioSpectrogramFrameData& SpectrogramFrameData, const EAudioSpectrogramFrequencyAxisPixelBucketMode InFrequencyAxisPixelBucketMode = EAudioSpectrogramFrequencyAxisPixelBucketMode::Sample);
+
+		// Include a move constructor due to being non-copyable.
+		FHistoryFrameData(FHistoryFrameData&& Other);
+
+		// Generate pixel data for this frame according to the given render params.
+		void UpdateCachedPixels(const FAudioSpectrogramViewportRenderParams& InRenderParams);
+
+		// Read the last generated pixel data.
+		TConstArrayView<FColor> GetCachedPixels() const { return CachedPixels; }
+
+	private:
+		// Allocate the necessary space for the mip chain if required.
+		void SetMipChainLengthUninitialized(const EAudioSpectrogramFrequencyAxisPixelBucketMode InFrequencyAxisPixelBucketMode);
+
+		// Generates mip maps if required.
+		void UpdateMipChain(const EAudioSpectrogramFrequencyAxisPixelBucketMode InFrequencyAxisPixelBucketMode);
+
+		// We have mips if the mip chain is larger than mip zero.
+		bool HasMips() const { return (SoundLevelsMipChain.Num() > NumSoundLevels); }
+
+		// Get a read-only view of the mip data for the given mip size.
+		TConstArrayView<float> GetMipFromMipSize(const int MipSize) const;
+
+		// Sample a fully interpolated value from possibly between two mips, at the given normalized read position.
+		float GetMipInterpolatedSoundLevel(const float DesiredMipSize, const float NormalizedDataPos) const;
+
+		// Helper function for sampling from the given mip data at the given normalized read position. Data is linearly interpolated.
+		static float GetInterpolatedSoundLevel(const TConstArrayView<float> MipData, const float NormalizedDataPos);
+
+		// Immutable data describing this spectrogram frame:
+		const int32 NumSoundLevels;
+		const float SpectrumDataMinFrequency;
+		const float SpectrumDataMaxFrequency;
+		const bool bSpectrumDataHasLogSpacedFreqencies;
+
+		// The first NumSoundLevels values are the full resolution spectrogram data for this frame, and considered immutable. Following that may also be lower resolution mips generated for the given current pixel bucket mode.
+		TArray<float> SoundLevelsMipChain;
+		EAudioSpectrogramFrequencyAxisPixelBucketMode SoundLevelsMipChainPixelBucketMode;
+
+		// The most recently generated pixel data.
+		TArray<FColor> CachedPixels;
+	};
+
+	FAudioSpectrogramViewportRenderParams RenderParams;
+	TRingBuffer<FHistoryFrameData> History;
+	FSlateUpdatableTexture* UpdatableTexture = nullptr;
+	EPixelFormat UpdatableTexturePixelFormat = PF_Unknown;
+};
+
diff --git a/Engine/Plugins/Runtime/AudioWidgets/Source/AudioWidgets/Public/SAudioSpectrogram.h b/Engine/Plugins/Runtime/AudioWidgets/Source/AudioWidgets/Public/SAudioSpectrogram.h
new file mode 100644
--- /dev/null
+++ b/Engine/Plugins/Runtime/AudioWidgets/Source/AudioWidgets/Public/SAudioSpectrogram.h
@@ -0,0 +1,65 @@
+// Copyright Epic Games, Inc. All Rights Reserved.
+
+#pragma once
+
+#include "AudioColorMapper.h"
+#include "AudioSpectrogramViewport.h"
+#include "Widgets/SCompoundWidget.h"
+
+/**
+ * Slate Widget for rendering a time-frequency representation of a series of audio power spectra.
+ */
+class AUDIOWIDGETS_API SAudioSpectrogram : public SCompoundWidget
+{
+public:
+	SLATE_BEGIN_ARGS(SAudioSpectrogram)
+		: _ViewMinFrequency(20.0f)
+		, _ViewMaxFrequency(20000.0f)
+		, _ColorMapMinSoundLevel(-84.0f)
+		, _ColorMapMaxSoundLevel(12.0f)
+		, _ColorMap(EAudioColorMap::BlackToWhite)
+		, _FrequencyAxisScale(EAudioSpectrogramFrequencyAxisScale::Logarithmic)
+		, _FrequencyAxisPixelBucketMode(EAudioSpectrogramFrequencyAxisPixelBucketMode::Average)
+		, _Orientation(EOrientation::Orient_Horizontal)
+	{}
+		SLATE_ATTRIBUTE(float, ViewMinFrequency)
+		SLATE_ATTRIBUTE(float, ViewMaxFrequency)
+		SLATE_ATTRIBUTE(float, ColorMapMinSoundLevel)
+		SLATE_ATTRIBUTE(float, ColorMapMaxSoundLevel)
+		SLATE_ATTRIBUTE(EAudioColorMap, ColorMap)
+		SLATE_ATTRIBUTE(EAudioSpectrogramFrequencyAxisScale, FrequencyAxisScale)
+		SLATE_ATTRIBUTE(EAudioSpectrogramFrequencyAxisPixelBucketMode, FrequencyAxisPixelBucketMode)
+		SLATE_ATTRIBUTE(EOrientation, Orientation)
+	SLATE_END_ARGS()
+
+	/** Constructs this widget with InArgs */
+	void Construct(const FArguments& InArgs);
+
+	/** Add the data for one spectrum frame to the spectrogram display */
+	void AddFrame(const FAudioSpectrogramFrameData& SpectrogramFrameData);
+
+	void SetViewMinFrequency(const float InViewMinFrequency) { ViewMinFrequency = InViewMinFrequency; }
+	void SetViewMaxFrequency(const float InViewMaxFrequency) { ViewMaxFrequency = InViewMaxFrequency; }
+	void SetColorMapMinSoundLevel(const float InColorMapMinSoundLevel) { ColorMapMinSoundLevel = InColorMapMinSoundLevel; }
+	void SetColorMapMaxSoundLevel(const float InColorMapMaxSoundLevel) { ColorMapMaxSoundLevel = InColorMapMaxSoundLevel; }
+	void SetColorMap(const EAudioColorMap InColorMap) { ColorMap = InColorMap; }
+	void SetFrequencyAxisScale(const EAudioSpectrogramFrequencyAxisScale InFrequencyAxisScale) { FrequencyAxisScale = InFrequencyAxisScale; }
+	void SetFrequencyAxisPixelBucketMode(const EAudioSpectrogramFrequencyAxisPixelBucketMode InFrequencyAxisPixelBucketMode) { FrequencyAxisPixelBucketMode = InFrequencyAxisPixelBucketMode; }
+	void SetOrientation(const EOrientation InOrientation) { Orientation = InOrientation; }
+
+private:
+	// Begin SWidget overrides.
+	virtual int32 OnPaint(const FPaintArgs& Args, const FGeometry& AllottedGeometry, const FSlateRect& MyCullingRect, FSlateWindowElementList& OutDrawElements, int32 LayerId, const FWidgetStyle& InWidgetStyle, bool bParentEnabled) const override;
+	// End SWidget overrides.
+
+	TAttribute<float> ViewMinFrequency;
+	TAttribute<float> ViewMaxFrequency;
+	TAttribute<float> ColorMapMinSoundLevel;
+	TAttribute<float> ColorMapMaxSoundLevel;
+	TAttribute<EAudioColorMap> ColorMap;
+	TAttribute<EAudioSpectrogramFrequencyAxisScale> FrequencyAxisScale;
+	TAttribute<EAudioSpectrogramFrequencyAxisPixelBucketMode> FrequencyAxisPixelBucketMode;
+	TAttribute<EOrientation> Orientation;
+
+	TSharedPtr<FAudioSpectrogramViewport> SpectrogramViewport;
+};
