Added the new AudioMaterialSlate

#virtualized

diff --git a/Engine/Plugins/Runtime/AudioWidgets/Source/AudioWidgets/Private/AudioMaterialSlate/AudioMaterialButton.cpp b/Engine/Plugins/Runtime/AudioWidgets/Source/AudioWidgets/Private/AudioMaterialSlate/AudioMaterialButton.cpp
new file mode 100644
--- /dev/null
+++ b/Engine/Plugins/Runtime/AudioWidgets/Source/AudioWidgets/Private/AudioMaterialSlate/AudioMaterialButton.cpp
@@ -0,0 +1,70 @@
+// Copyright Epic Games, Inc. All Rights Reserved.
+
+
+#include "AudioMaterialSlate/AudioMaterialButton.h"
+#include "AudioMaterialSlate/SAudioMaterialButton.h"
+
+#define LOCTEXT_NAMESPACE "AudioWidgets"
+#if WITH_EDITOR
+const FText UAudioMaterialButton::GetPaletteCategory()
+{
+	return LOCTEXT("PaletteCategory", "AudioMaterial");
+}
+#endif // WITH_EDITOR
+
+void UAudioMaterialButton::SynchronizeProperties()
+{
+	Super::SynchronizeProperties();
+
+	if (!Button.IsValid())
+	{
+		return;
+	}
+
+	Button->SetPressedState(bIsPressed);
+	Button->ApplyNewMaterial();
+}
+
+void UAudioMaterialButton::ReleaseSlateResources(bool bReleaseChildren)
+{
+	Super::ReleaseSlateResources(bReleaseChildren);
+	Button.Reset();
+}
+
+bool UAudioMaterialButton::GetIsPressed() const
+{
+	return bIsPressed;
+}
+
+void UAudioMaterialButton::SetIsPressed(bool InPressed)
+{
+	if (!Button.IsValid())
+	{
+		return;
+	}
+
+	if (bIsPressed != InPressed)
+	{
+		bIsPressed = InPressed;
+		OnButtonPressedChangedEvent.Broadcast(InPressed);
+	}
+}
+
+TSharedRef<SWidget> UAudioMaterialButton::RebuildWidget()
+{
+	Button = SNew(SAudioMaterialButton)
+		.Owner(this)
+		.AudioMaterialButtonStyle(&WidgetStyle)
+		.bIsPressedAttribute(bIsPressed)
+		.OnBooleanValueChanged(BIND_UOBJECT_DELEGATE(FOnBooleanValueChanged, HandleOnPressedValueChanged));
+
+	return Button.ToSharedRef();
+}
+
+void UAudioMaterialButton::HandleOnPressedValueChanged(bool InPressedSate)
+{
+	bIsPressed = InPressedSate;
+	OnButtonPressedChangedEvent.Broadcast(InPressedSate);
+}
+
+#undef LOCTEXT_NAMESPACE
diff --git a/Engine/Plugins/Runtime/AudioWidgets/Source/AudioWidgets/Private/AudioMaterialSlate/AudioMaterialKnob.cpp b/Engine/Plugins/Runtime/AudioWidgets/Source/AudioWidgets/Private/AudioMaterialSlate/AudioMaterialKnob.cpp
new file mode 100644
--- /dev/null
+++ b/Engine/Plugins/Runtime/AudioWidgets/Source/AudioWidgets/Private/AudioMaterialSlate/AudioMaterialKnob.cpp
@@ -0,0 +1,72 @@
+// Copyright Epic Games, Inc. All Rights Reserved.
+
+
+#include "AudioMaterialSlate/AudioMaterialKnob.h"
+#include "AudioMaterialSlate/SAudioMaterialKnob.h"
+
+#define LOCTEXT_NAMESPACE "AudioWidgets"
+#if WITH_EDITOR
+const FText UAudioMaterialKnob::GetPaletteCategory()
+{
+	return LOCTEXT("PaletteCategory", "AudioMaterial");
+}
+#endif // WITH_EDITOR
+
+void UAudioMaterialKnob::SynchronizeProperties()
+{
+	Super::SynchronizeProperties();
+
+	if (!Knob.IsValid())
+	{
+		return;
+	}
+
+	Knob->SetValue(Value);
+	Knob->ApplyNewMaterial();
+}
+
+void UAudioMaterialKnob::ReleaseSlateResources(bool bReleaseChildren)
+{
+	Super::ReleaseSlateResources(bReleaseChildren);
+	Knob.Reset();
+}
+
+float UAudioMaterialKnob::GetValue()
+{
+	return Value;
+}
+
+void UAudioMaterialKnob::SetValue(float InValue)
+{
+	InValue = FMath::Clamp(InValue, 0.f, 1.f);
+
+	if (Knob.IsValid())
+	{
+		Knob->SetValue(InValue);
+	}
+
+	if (Value != InValue)
+	{
+		Value = InValue;
+		HandleOnKnobValueChanged(InValue);
+	}
+}
+
+TSharedRef<SWidget> UAudioMaterialKnob::RebuildWidget()
+{
+	Knob = SNew(SAudioMaterialKnob)
+		.Owner(this)
+		.AudioMaterialKnobStyle(&WidgetStyle)
+		.Value(Value)
+		.OnFloatValueChanged(BIND_UOBJECT_DELEGATE(FOnFloatValueChanged, HandleOnKnobValueChanged));
+		
+	return Knob.ToSharedRef();
+}
+
+void UAudioMaterialKnob::HandleOnKnobValueChanged(float InValue)
+{
+	Value = InValue;
+	OnKnobValueChanged.Broadcast(InValue);	
+}
+
+#undef LOCTEXT_NAMESPACE
diff --git a/Engine/Plugins/Runtime/AudioWidgets/Source/AudioWidgets/Private/AudioMaterialSlate/AudioMaterialSlateStyles.cpp b/Engine/Plugins/Runtime/AudioWidgets/Source/AudioWidgets/Private/AudioMaterialSlate/AudioMaterialSlateStyles.cpp
new file mode 100644
--- /dev/null
+++ b/Engine/Plugins/Runtime/AudioWidgets/Source/AudioWidgets/Private/AudioMaterialSlate/AudioMaterialSlateStyles.cpp
@@ -0,0 +1,96 @@
+// Copyright Epic Games, Inc. All Rights Reserved.
+
+
+#include "AudioMaterialSlate/AudioMaterialSlateStyles.h"
+
+TObjectPtr<UMaterialInstanceDynamic> FAudioMaterialStyle::GetDynamicMaterial() const
+{
+	return DynamicMaterial;
+}
+
+void FAudioMaterialStyle::CreateDynamicMaterial(UObject* InOuter) const
+{
+	DynamicMaterial = UMaterialInstanceDynamic::Create(Material, InOuter);
+}
+
+void FAudioMaterialSliderStyle::GetResources(TArray<const FSlateBrush*>& OutBrushes) const
+{
+	// Add any brush resources here so that Slate can correctly atlas and reference them
+}
+
+FAudioMaterialButtonStyle::FAudioMaterialButtonStyle()
+:ButtonMainColor(FLinearColor(0.1f, 0.1f, 0.1f, 1.0f)),
+ButtonShadowColor(FLinearColor(0.0f, 0.0f, 0.0f, 1.0f)),
+ButtonAccentColor(FLinearColor(0.34f, 0.34f, 0.34f, 1.0f)),
+ButtonPressedMainColor(FLinearColor(0.95f, 0.95f, 0.95f, 1.0f)),
+ButtonShadowMainColor(FLinearColor(0.156f, 0.138f, 0.15f, 1.0f)),
+ButtonPressedOutlineColor(FLinearColor(0.156f, 0.138f, 0.15f, 1.0f))
+{
+}
+
+FAudioMaterialButtonStyle::~FAudioMaterialButtonStyle()
+{
+}
+
+const FName FAudioMaterialButtonStyle::TypeName(TEXT("FAudioMaterialButtonStyle"));
+
+void FAudioMaterialButtonStyle::GetResources(TArray<const FSlateBrush*>& OutBrushes) const
+{
+}
+
+const FAudioMaterialButtonStyle& FAudioMaterialButtonStyle::GetDefault()
+{
+	static FAudioMaterialButtonStyle Default;
+	return Default;
+}
+
+FAudioMaterialSliderStyle::FAudioMaterialSliderStyle()
+:BarMainColor(FLinearColor(1.0f, 1.0f, 1.0f, 0.0f)),
+BarShadowColor(FLinearColor(1.0f, 1.0f, 1.0f, 0.0f)),
+BarAccentColor(FLinearColor(1.0f, 1.0f, 1.0f, 0.0f)),
+HandleMainColor(FLinearColor(1.0f, 1.0f, 1.0f, 0.0f)),
+HandleOutlineColor(FLinearColor(1.0f, 1.0f, 1.0f, 0.0f))
+{
+}
+
+FAudioMaterialSliderStyle::~FAudioMaterialSliderStyle()
+{
+}
+
+const FName FAudioMaterialSliderStyle::TypeName(TEXT("FAudioMaterialSliderStyle"));
+
+const FAudioMaterialSliderStyle& FAudioMaterialSliderStyle::GetDefault()
+{
+	static FAudioMaterialSliderStyle Default;
+	return Default;
+}
+
+FAudioMaterialKnobStyle::FAudioMaterialKnobStyle()
+:KnobMainColor(FLinearColor(0.f, 0.f, 0.f, 1.0f)),
+KnobAccentColor(FLinearColor(.14f, .16f, .2f, 1.0f)),
+KnobIndicatorColor(FLinearColor(1.f, 0.f, 0.04f, 1.0f)),
+KnobBarColor(FLinearColor(.033f, .033f, .033f, 1.0f)),
+KnobBarShadowColor(FLinearColor(0.f, 0.f, 0.f, 1.0f)),
+KnobBarFillMinColor(FLinearColor(0.9f, 0.9f, 0.9f, 1.0f)),
+KnobBarFillMidColor(FLinearColor(0.9f, 0.9f, 0.9f, 1.0f)),
+KnobBarFillMaxColor(FLinearColor(0.9f, 0.9f, 0.9f, 1.0f)),
+KnobBarFillTintColor(FLinearColor(1.f, 1.f, 1.f, 1.0f))
+{
+}
+
+FAudioMaterialKnobStyle::~FAudioMaterialKnobStyle()
+{
+}
+
+const FName FAudioMaterialKnobStyle::TypeName(TEXT("FAudioMaterialKnobStyle"));
+
+const FAudioMaterialKnobStyle& FAudioMaterialKnobStyle::GetDefault()
+{
+	static FAudioMaterialKnobStyle Default;
+	return Default;
+}
+
+void FAudioMaterialKnobStyle::GetResources(TArray<const FSlateBrush*>& OutBrushes) const
+{
+	// Add any brush resources here so that Slate can correctly atlas and reference them
+}
diff --git a/Engine/Plugins/Runtime/AudioWidgets/Source/AudioWidgets/Private/AudioMaterialSlate/AudioMaterialSlider.cpp b/Engine/Plugins/Runtime/AudioWidgets/Source/AudioWidgets/Private/AudioMaterialSlate/AudioMaterialSlider.cpp
new file mode 100644
--- /dev/null
+++ b/Engine/Plugins/Runtime/AudioWidgets/Source/AudioWidgets/Private/AudioMaterialSlate/AudioMaterialSlider.cpp
@@ -0,0 +1,72 @@
+// Copyright Epic Games, Inc. All Rights Reserved.
+
+
+#include "AudioMaterialSlate/AudioMaterialSlider.h"
+#include "AudioMaterialSlate/AudioMaterialSlateStyles.h"
+#include "AudioMaterialSlate/SAudioMaterialSlider.h"
+#include "Widgets/SWeakWidget.h"
+
+#define LOCTEXT_NAMESPACE "AudioWidgets"
+
+#if WITH_EDITOR
+const FText UAudioMaterialSlider::GetPaletteCategory()
+{
+	return LOCTEXT("PaletteCategory", "AudioMaterial");
+}
+#endif
+
+void UAudioMaterialSlider::SynchronizeProperties()
+{
+	Super::SynchronizeProperties();
+
+	if (!Slider.IsValid())
+	{
+		return;
+	}
+
+	Slider->SetValue(Value);
+	Slider->ApplyNewMaterial();
+}
+
+void UAudioMaterialSlider::ReleaseSlateResources(bool bReleaseChildren)
+{
+	Slider.Reset();
+}
+
+float UAudioMaterialSlider::GetValue() const
+{
+	return Value;
+}
+
+void UAudioMaterialSlider::SetValue(float InValue)
+{
+	InValue = FMath::Clamp(InValue, 0.f, 1.f);
+
+	if (Slider.IsValid())
+	{
+		Slider->SetValue(InValue);
+	}
+
+	if (Value != InValue)
+	{
+		Value = InValue;
+		HandleOnValueChanged(InValue);
+	}
+}
+
+TSharedRef<SWidget> UAudioMaterialSlider::RebuildWidget()
+{
+	Slider = SNew(SAudioMaterialSlider)
+		.Owner(this)
+		.AudioMaterialSliderStyle(&WidgetStyle)
+		.OnValueChanged(BIND_UOBJECT_DELEGATE(FOnFloatValueChanged, HandleOnValueChanged));
+
+	return Slider.ToSharedRef();
+}
+
+void UAudioMaterialSlider::HandleOnValueChanged(float InValue)
+{
+	OnValueChanged.Broadcast(InValue);
+}
+
+#undef LOCTEXT_NAMESPACE
diff --git a/Engine/Plugins/Runtime/AudioWidgets/Source/AudioWidgets/Private/AudioMaterialSlate/SAudioMaterialButton.cpp b/Engine/Plugins/Runtime/AudioWidgets/Source/AudioWidgets/Private/AudioMaterialSlate/SAudioMaterialButton.cpp
new file mode 100644
--- /dev/null
+++ b/Engine/Plugins/Runtime/AudioWidgets/Source/AudioWidgets/Private/AudioMaterialSlate/SAudioMaterialButton.cpp
@@ -0,0 +1,97 @@
+// Copyright Epic Games, Inc. All Rights Reserved.
+
+
+#include "AudioMaterialSlate/SAudioMaterialButton.h"
+#include "SlateOptMacros.h"
+#include "Components/Widget.h"
+
+BEGIN_SLATE_FUNCTION_BUILD_OPTIMIZATION
+void SAudioMaterialButton::Construct(const FArguments& InArgs)
+{
+	Owner = InArgs._Owner;
+	AudioMaterialButtonStyle = InArgs._AudioMaterialButtonStyle;
+	bIsPressedAttribute = InArgs._bIsPressedAttribute;
+	OnIsPressedStateChanged = InArgs._OnBooleanValueChanged;
+
+	ApplyNewMaterial();
+
+}
+END_SLATE_FUNCTION_BUILD_OPTIMIZATION
+
+void SAudioMaterialButton::SetPressedState(bool InPressedState)
+{
+	CommitNewState(InPressedState);
+}
+
+void SAudioMaterialButton::ApplyNewMaterial()
+{
+	if (AudioMaterialButtonStyle)
+	{
+		AudioMaterialButtonStyle->CreateDynamicMaterial(Owner.Get());
+	}
+}
+
+int32 SAudioMaterialButton::OnPaint(const FPaintArgs& Args, const FGeometry& AllottedGeometry, const FSlateRect& MyCullingRect, FSlateWindowElementList& OutDrawElements, int32 LayerId, const FWidgetStyle& InWidgetStyle, bool bParentEnabled) const
+{
+	if (AudioMaterialButtonStyle)
+	{
+		UMaterialInstanceDynamic* DynamicMaterial = AudioMaterialButtonStyle->GetDynamicMaterial();
+
+		if (IsValid(DynamicMaterial))
+		{
+			DynamicMaterial->SetVectorParameterValue(FName("MainColor"), AudioMaterialButtonStyle->ButtonMainColor);
+			DynamicMaterial->SetVectorParameterValue(FName("ShadowColor"), AudioMaterialButtonStyle->ButtonShadowColor);
+			DynamicMaterial->SetVectorParameterValue(FName("SmoothBevelColor"), AudioMaterialButtonStyle->ButtonAccentColor);
+			DynamicMaterial->SetVectorParameterValue(FName("Color_1"), AudioMaterialButtonStyle->ButtonPressedMainColor);
+			DynamicMaterial->SetVectorParameterValue(FName("Color_2"), AudioMaterialButtonStyle->ButtonShadowMainColor);
+			DynamicMaterial->SetVectorParameterValue(FName("LedColor"), AudioMaterialButtonStyle->ButtonPressedOutlineColor);
+			DynamicMaterial->SetScalarParameterValue(FName("Click"), bIsPressedAttribute.Get());
+		}
+
+		const bool bEnabled = ShouldBeEnabled(bParentEnabled);
+		const ESlateDrawEffect DrawEffects = bEnabled ? ESlateDrawEffect::None : ESlateDrawEffect::DisabledEffect;
+
+		const FLinearColor FinalColorAndOpacity(InWidgetStyle.GetColorAndOpacityTint());
+
+		FSlateBrush Brush;
+		Brush.SetResourceObject(DynamicMaterial);
+		FSlateDrawElement::MakeBox(OutDrawElements, LayerId++, AllottedGeometry.ToPaintGeometry(),&Brush, DrawEffects, FinalColorAndOpacity);
+	}
+
+	return LayerId;
+}
+
+FVector2D SAudioMaterialButton::ComputeDesiredSize(float) const
+{
+	FVector2D Vector = FVector2D(SlateWidth, SlateHeight);
+	return Vector;
+}
+
+FReply SAudioMaterialButton::OnMouseButtonDown(const FGeometry& MyGeometry, const FPointerEvent& MouseEvent)
+{
+	if ((MouseEvent.GetEffectingButton() == EKeys::LeftMouseButton))
+	{
+		CommitNewState(!bIsPressedAttribute.Get());
+
+		FReply Reply = FReply::Handled().CaptureMouse(SharedThis(this));
+		return Reply;
+	}
+
+	return FReply::Unhandled();
+}
+
+FReply SAudioMaterialButton::OnMouseButtonUp(const FGeometry& MyGeometry, const FPointerEvent& MouseEvent)
+{
+	if ((MouseEvent.GetEffectingButton() == EKeys::LeftMouseButton) && this->HasMouseCapture())
+	{
+		return FReply::Handled().ReleaseMouseCapture();
+	}
+
+	return FReply::Unhandled();
+}
+
+void SAudioMaterialButton::CommitNewState(bool InPressedState)
+{
+	bIsPressedAttribute.Set(InPressedState);
+	OnIsPressedStateChanged.ExecuteIfBound(InPressedState);
+}
diff --git a/Engine/Plugins/Runtime/AudioWidgets/Source/AudioWidgets/Private/AudioMaterialSlate/SAudioMaterialKnob.cpp b/Engine/Plugins/Runtime/AudioWidgets/Source/AudioWidgets/Private/AudioMaterialSlate/SAudioMaterialKnob.cpp
new file mode 100644
--- /dev/null
+++ b/Engine/Plugins/Runtime/AudioWidgets/Source/AudioWidgets/Private/AudioMaterialSlate/SAudioMaterialKnob.cpp
@@ -0,0 +1,126 @@
+// Copyright Epic Games, Inc. All Rights Reserved.
+
+
+#include "AudioMaterialSlate/SAudioMaterialKnob.h"
+#include "SlateOptMacros.h"
+#include "Components/Widget.h"
+
+BEGIN_SLATE_FUNCTION_BUILD_OPTIMIZATION
+void SAudioMaterialKnob::Construct(const FArguments& InArgs)
+{
+	Owner = InArgs._Owner;
+	AudioMaterialKnobStyle = InArgs._AudioMaterialKnobStyle;
+	OnValueChanged = InArgs._OnFloatValueChanged;
+
+	ApplyNewMaterial();
+	CommitValue(InArgs._Value.Get());
+}
+END_SLATE_FUNCTION_BUILD_OPTIMIZATION
+
+void SAudioMaterialKnob::SetValue(float InValue)
+{
+	CommitValue(InValue);
+}
+
+void SAudioMaterialKnob::ApplyNewMaterial()
+{
+	if (AudioMaterialKnobStyle)
+	{
+		AudioMaterialKnobStyle->CreateDynamicMaterial(Owner.Get());
+	}
+}
+
+int32 SAudioMaterialKnob::OnPaint(const FPaintArgs& Args, const FGeometry& AllottedGeometry, const FSlateRect& MyCullingRect, FSlateWindowElementList& OutDrawElements, int32 LayerId, const FWidgetStyle& InWidgetStyle, bool bParentEnabled) const
+{
+	if (AudioMaterialKnobStyle)
+	{
+		UMaterialInstanceDynamic* DynamicMaterial = AudioMaterialKnobStyle->GetDynamicMaterial();
+
+		if (IsValid(DynamicMaterial))
+		{
+			const float KnobPercent = ValueAttribute.Get();
+
+			//TODO unify material parameter names
+			DynamicMaterial->SetVectorParameterValue(FName("Color_2"), AudioMaterialKnobStyle->KnobMainColor);
+			DynamicMaterial->SetVectorParameterValue(FName("Color_1"), AudioMaterialKnobStyle->KnobAccentColor);
+
+			DynamicMaterial->SetVectorParameterValue(FName("BarColor"), AudioMaterialKnobStyle->KnobBarColor);
+			DynamicMaterial->SetVectorParameterValue(FName("BarShadowColor"), AudioMaterialKnobStyle->KnobBarShadowColor);
+			DynamicMaterial->SetVectorParameterValue(FName("LEDglowMax"), AudioMaterialKnobStyle->KnobBarFillMaxColor);
+			DynamicMaterial->SetVectorParameterValue(FName("Led_Max"), AudioMaterialKnobStyle->KnobBarFillMaxColor);
+			DynamicMaterial->SetVectorParameterValue(FName("LEDglowMid"), AudioMaterialKnobStyle->KnobBarFillMidColor);
+			DynamicMaterial->SetVectorParameterValue(FName("Led_Mid"), AudioMaterialKnobStyle->KnobBarFillMidColor);
+			DynamicMaterial->SetVectorParameterValue(FName("LEDglowMin"), AudioMaterialKnobStyle->KnobBarFillMinColor);
+			DynamicMaterial->SetVectorParameterValue(FName("LED_Min"), AudioMaterialKnobStyle->KnobBarFillMinColor);
+			DynamicMaterial->SetVectorParameterValue(FName("LineColor"), AudioMaterialKnobStyle->KnobIndicatorColor);
+			DynamicMaterial->SetVectorParameterValue(FName("LedTint"), AudioMaterialKnobStyle->KnobBarFillTintColor);
+
+			DynamicMaterial->SetScalarParameterValue(FName("VALUE"), FMath::Clamp(KnobPercent, 0.f, 1.f));
+		}
+
+		const bool bEnabled = ShouldBeEnabled(bParentEnabled);
+		const ESlateDrawEffect DrawEffects = bEnabled ? ESlateDrawEffect::None : ESlateDrawEffect::DisabledEffect;
+
+		const FLinearColor FinalColorAndOpacity(InWidgetStyle.GetColorAndOpacityTint());
+
+		FSlateBrush Brush;
+		Brush.SetResourceObject(DynamicMaterial);
+		FSlateDrawElement::MakeBox(OutDrawElements, LayerId++, AllottedGeometry.ToPaintGeometry(),&Brush, DrawEffects, FinalColorAndOpacity);
+	}
+
+	return LayerId;
+}
+
+FVector2D SAudioMaterialKnob::ComputeDesiredSize(float) const
+{
+	FVector2D Vector = FVector2D(SlateWidth, SlateHeight);
+	return Vector;
+}
+
+FReply SAudioMaterialKnob::OnMouseMove(const FGeometry& MyGeometry, const FPointerEvent& MouseEvent)
+{
+	if (this->HasMouseCapture())
+	{
+		int32 CurrentYValue = MouseEvent.GetLastScreenSpacePosition().Y;
+
+		const float MouseSpeed = 0.2f;
+
+		float ValueDelta = (float)(MouseDownPosition.Y - CurrentYValue) / PixelDelta * MouseSpeed;
+		float NewValue = FMath::Clamp(MouseDownValue + ValueDelta, 0.0f, 1.0f);
+		CommitValue(NewValue);
+
+		return FReply::Handled();
+	}
+
+	return FReply::Unhandled();
+}
+
+FReply SAudioMaterialKnob::OnMouseButtonDown(const FGeometry& MyGeometry, const FPointerEvent& MouseEvent)
+{
+	if ((MouseEvent.GetEffectingButton() == EKeys::LeftMouseButton))
+	{
+		MouseDownPosition = MouseEvent.GetScreenSpacePosition();
+		MouseDownValue = ValueAttribute.Get();
+
+		FReply Reply = FReply::Handled().CaptureMouse(SharedThis(this));
+		return Reply;
+	}
+
+	return FReply::Unhandled();
+}
+
+FReply SAudioMaterialKnob::OnMouseButtonUp(const FGeometry& MyGeometry, const FPointerEvent& MouseEvent)
+{
+	if ((MouseEvent.GetEffectingButton() == EKeys::LeftMouseButton) && this->HasMouseCapture())
+	{
+		return FReply::Handled().ReleaseMouseCapture();
+	}
+
+	return FReply::Unhandled();
+}
+
+void SAudioMaterialKnob::CommitValue(float NewValue)
+{
+	ValueAttribute.Set(NewValue);
+	OnValueChanged.ExecuteIfBound(NewValue);
+}
diff --git a/Engine/Plugins/Runtime/AudioWidgets/Source/AudioWidgets/Private/AudioMaterialSlate/SAudioMaterialSlider.cpp b/Engine/Plugins/Runtime/AudioWidgets/Source/AudioWidgets/Private/AudioMaterialSlate/SAudioMaterialSlider.cpp
new file mode 100644
--- /dev/null
+++ b/Engine/Plugins/Runtime/AudioWidgets/Source/AudioWidgets/Private/AudioMaterialSlate/SAudioMaterialSlider.cpp
@@ -0,0 +1,132 @@
+// Copyright Epic Games, Inc. All Rights Reserved.
+
+
+#include "AudioMaterialSlate/SAudioMaterialSlider.h"
+#include "AudioMaterialSlate/AudioMaterialSlider.h"
+#include "SlateOptMacros.h"
+#include "Styling/SlateBrush.h"
+
+BEGIN_SLATE_FUNCTION_BUILD_OPTIMIZATION
+
+void SAudioMaterialSlider::Construct(const FArguments& InArgs)
+{
+	Owner = InArgs._Owner;
+	OnValueChanged = InArgs._OnValueChanged;
+	ValueAttribute = InArgs._ValueAttribute.Get();
+	AudioMaterialSliderStyle = InArgs._AudioMaterialSliderStyle;
+
+	ApplyNewMaterial();
+	CommitValue(InArgs._ValueAttribute.Get());
+}
+
+END_SLATE_FUNCTION_BUILD_OPTIMIZATION
+
+int32 SAudioMaterialSlider::OnPaint(const FPaintArgs& Args, const FGeometry& AllottedGeometry, const FSlateRect& MyCullingRect, FSlateWindowElementList& OutDrawElements, int32 LayerId, const FWidgetStyle& InWidgetStyle, bool bParentEnabled) const
+{
+
+	if (AudioMaterialSliderStyle)
+	{
+		UMaterialInstanceDynamic* DynamicMaterial = AudioMaterialSliderStyle->GetDynamicMaterial();
+		if (IsValid(DynamicMaterial))
+		{
+			const float Value = ValueAttribute.Get();
+			DynamicMaterial->SetVectorParameterValue(FName("BarColor"), AudioMaterialSliderStyle->BarMainColor);
+			DynamicMaterial->SetVectorParameterValue(FName("LedColor"), AudioMaterialSliderStyle->BarShadowColor);
+			DynamicMaterial->SetVectorParameterValue(FName("BarInnerShadow"), AudioMaterialSliderStyle->BarAccentColor);
+			DynamicMaterial->SetVectorParameterValue(FName("ValueColor"), AudioMaterialSliderStyle->HandleMainColor);
+			DynamicMaterial->SetVectorParameterValue(FName("DotBevel1"), AudioMaterialSliderStyle->HandleOutlineColor);
+			DynamicMaterial->SetScalarParameterValue(FName("VALUE"), FMath::Clamp(Value, 0.f, 1.f));
+		}
+
+		const bool bEnabled = ShouldBeEnabled(bParentEnabled);
+		const ESlateDrawEffect DrawEffects = bEnabled ? ESlateDrawEffect::None : ESlateDrawEffect::DisabledEffect;
+			
+		const FLinearColor FinalColorAndOpacity(InWidgetStyle.GetColorAndOpacityTint());
+
+		FSlateBrush Brush;
+		Brush.SetResourceObject(DynamicMaterial);
+		FSlateDrawElement::MakeBox(OutDrawElements, LayerId++, AllottedGeometry.ToPaintGeometry(), &Brush, DrawEffects, FinalColorAndOpacity);
+	}
+
+	return LayerId;
+}
+
+FVector2D SAudioMaterialSlider::ComputeDesiredSize(float) const
+{
+	FVector2D Vector = FVector2D(SlateWidth, SlateHeight);
+	return Vector;
+}
+
+void SAudioMaterialSlider::SetValue(TAttribute<float> InValueAttribute)
+{
+	CommitValue(InValueAttribute.Get());
+}
+
+void SAudioMaterialSlider::ApplyNewMaterial()
+{
+	if (AudioMaterialSliderStyle)
+	{
+		AudioMaterialSliderStyle->CreateDynamicMaterial(Owner.Get());
+	}
+}
+
+FReply SAudioMaterialSlider::OnMouseButtonDown(const FGeometry& MyGeometry, const FPointerEvent& MouseEvent)
+{
+	if ((MouseEvent.GetEffectingButton() == EKeys::LeftMouseButton))
+	{
+		FVector2D Value2D = PositionToValue(MyGeometry, MouseEvent.GetLastScreenSpacePosition());
+		CommitValue(Value2D.Y);
+
+		return FReply::Handled().CaptureMouse(SharedThis(this));
+	}
+
+	return FReply::Unhandled();
+}
+
+FReply SAudioMaterialSlider::OnMouseButtonUp(const FGeometry& MyGeometry, const FPointerEvent& MouseEvent)
+{
+	if ((MouseEvent.GetEffectingButton() == EKeys::LeftMouseButton) && HasMouseCapture())
+	{
+		return FReply::Handled().ReleaseMouseCapture();
+	}
+
+	return FReply::Unhandled();
+}
+
+FReply SAudioMaterialSlider::OnMouseMove(const FGeometry& MyGeometry, const FPointerEvent& MouseEvent)
+{
+	if (this->HasMouseCapture())
+	{
+		FVector2D Vector2D = PositionToValue(MyGeometry, MouseEvent.GetLastScreenSpacePosition());
+
+		CommitValue(Vector2D.Y);
+
+		return FReply::Handled();
+	}
+
+	return FReply::Unhandled();
+}
+
+void SAudioMaterialSlider::CommitValue(float NewValue)
+{
+	float Val = FMath::Clamp(NewValue, 0.f , 1.f);
+	if (!ValueAttribute.IsBound())
+	{
+		ValueAttribute.Set(Val);
+	}
+
+	OnValueChanged.ExecuteIfBound(Val);
+}
+
+FVector2D SAudioMaterialSlider::PositionToValue(const FGeometry& MyGeometry, const FVector2D& AbsolutePosition)
+{
+	const FVector2D LocalPosition = MyGeometry.AbsoluteToLocal(AbsolutePosition);
+
+	float RelativeValueX = 1 - (LocalPosition.X) / (MyGeometry.Size.X);
+	float RelativeValueY = 1 - (LocalPosition.Y) / (MyGeometry.Size.Y);
+
+	RelativeValueX = FMath::Clamp(RelativeValueX, 0.0f, 1.0f);
+	RelativeValueY = FMath::Clamp(RelativeValueY, 0.0f, 1.0f);
+
+	return FVector2D(RelativeValueX, RelativeValueY);
+}
diff --git a/Engine/Plugins/Runtime/AudioWidgets/Source/AudioWidgets/Public/AudioMaterialSlate/AudioMaterialButton.h b/Engine/Plugins/Runtime/AudioWidgets/Source/AudioWidgets/Public/AudioMaterialSlate/AudioMaterialButton.h
new file mode 100644
--- /dev/null
+++ b/Engine/Plugins/Runtime/AudioWidgets/Source/AudioWidgets/Public/AudioMaterialSlate/AudioMaterialButton.h
@@ -0,0 +1,80 @@
+// Copyright Epic Games, Inc. All Rights Reserved.
+
+#pragma once
+
+#include "AudioMaterialSlateStyles.h"
+#include "Components/Widget.h"
+#include "Delegates/Delegate.h"
+#include "AudioMaterialButton.generated.h"
+
+class SAudioMaterialButton;
+struct FAudioMaterialButtonStyle;
+
+DECLARE_DYNAMIC_MULTICAST_DELEGATE_OneParam(FOnButtonPressedChangedEvent, bool, bIsPressed);
+
+
+/**
+ * A simple widget that shows a button
+ * Button is rendered by using material instead of texture.
+ *
+ * * No Children
+ */
+UCLASS()
+class AUDIOWIDGETS_API UAudioMaterialButton : public UWidget
+{
+	GENERATED_BODY()
+
+public:
+
+	/** The button's style */
+	UPROPERTY(EditAnywhere, BlueprintReadWrite, Category = "Style", meta = (DisplayName = "Style"))
+	FAudioMaterialButtonStyle WidgetStyle;
+
+public:
+#if WITH_EDITOR
+	virtual const FText GetPaletteCategory() override;
+#endif // WITH_EDITOR
+
+
+	// UWidget
+	virtual void SynchronizeProperties() override;
+	// End of UWidget
+
+	// UVisual
+	virtual void ReleaseSlateResources(bool bReleaseChildren) override;
+	// End of UVisual
+
+		/** Gets the current value of the slider.*/
+	UFUNCTION(BlueprintPure, Category = "Behavior")
+	bool GetIsPressed() const;
+
+	/** Sets the current value of the slider. InValue is Clamped between 0.f - 1.f */
+	UFUNCTION(BlueprintCallable, Category = "Behavior")
+	void SetIsPressed(bool InPressed);
+
+public:
+
+	/** Called when the value is changed by button. */
+	UPROPERTY(BlueprintAssignable, Category = "Widget Event")
+	FOnButtonPressedChangedEvent OnButtonPressedChangedEvent;
+
+protected:
+
+	// UWidget
+	virtual TSharedRef<SWidget> RebuildWidget() override;
+	// End of UWidget
+
+	void HandleOnPressedValueChanged(bool InPressedSate);
+
+private:
+
+	/**Default Value of the button*/
+	UPROPERTY(EditAnywhere, BlueprintSetter = SetIsPressed, BlueprintGetter = GetIsPressed, Category = "Appearance")
+	bool bIsPressed = false;
+
+private:
+
+	/** Native Slate Widget */
+	TSharedPtr<SAudioMaterialButton> Button;
+
+};
diff --git a/Engine/Plugins/Runtime/AudioWidgets/Source/AudioWidgets/Public/AudioMaterialSlate/AudioMaterialKnob.h b/Engine/Plugins/Runtime/AudioWidgets/Source/AudioWidgets/Public/AudioMaterialSlate/AudioMaterialKnob.h
new file mode 100644
--- /dev/null
+++ b/Engine/Plugins/Runtime/AudioWidgets/Source/AudioWidgets/Public/AudioMaterialSlate/AudioMaterialKnob.h
@@ -0,0 +1,79 @@
+// Copyright Epic Games, Inc. All Rights Reserved.
+
+#pragma once
+
+#include "AudioMaterialSlateStyles.h"
+#include "Components/Widget.h"
+#include "Delegates/Delegate.h"
+#include "AudioMaterialKnob.generated.h"
+
+class SAudioMaterialKnob;
+struct FAudioMaterialKnobStyle;
+
+DECLARE_DYNAMIC_MULTICAST_DELEGATE_OneParam(FOnKnobValueChangedEvent, float, Value);
+
+
+/**
+ * A simple widget that shows a turning Knob that allows you to control the value between 0..1.
+ * Knob is rendered by using material instead of texture.
+ *
+ * * No Children
+ */
+UCLASS()
+class AUDIOWIDGETS_API UAudioMaterialKnob : public UWidget
+{
+	GENERATED_BODY()
+
+public:
+
+	/** The button's style */
+	UPROPERTY(EditAnywhere, BlueprintReadWrite, Category = "Style", meta = (DisplayName = "Style"))
+	FAudioMaterialKnobStyle WidgetStyle;
+
+public:
+#if WITH_EDITOR
+	virtual const FText GetPaletteCategory() override;
+#endif // WITH_EDITOR
+
+	// UWidget
+	virtual void SynchronizeProperties() override;
+	// End of UWidget
+
+	// UVisual
+	virtual void ReleaseSlateResources(bool bReleaseChildren) override;
+	// End of UVisual
+
+	/** Gets the current value of the knob.*/
+	UFUNCTION(BlueprintPure, Category = "Behavior")
+	float GetValue();
+
+	/** Sets the current value of the knob. InValue is Clamped between 0.f - 1.f */
+	UFUNCTION(BlueprintCallable, Category = "Behavior")
+	void SetValue(float InValue);
+
+public:
+
+	/** Called when the value is changed by knob. */
+	UPROPERTY(BlueprintAssignable, Category = "Widget Event")
+	FOnKnobValueChangedEvent OnKnobValueChanged;
+
+protected:
+
+	// UWidget
+	virtual TSharedRef<SWidget> RebuildWidget() override;
+	// End of UWidget
+
+	void HandleOnKnobValueChanged(float InValue);
+
+private:
+
+	/**Default Value of the Knob*/
+	UPROPERTY(EditAnywhere, BlueprintSetter = SetValue, BlueprintGetter = GetValue, Category = "Appearance", meta = (UIMin = "0", UIMax = "1"))
+	float Value = 1.f;
+
+private:
+
+	/** Native Slate Widget */
+	TSharedPtr<SAudioMaterialKnob> Knob;
+
+};
diff --git a/Engine/Plugins/Runtime/AudioWidgets/Source/AudioWidgets/Public/AudioMaterialSlate/AudioMaterialSlateStyles.h b/Engine/Plugins/Runtime/AudioWidgets/Source/AudioWidgets/Public/AudioMaterialSlate/AudioMaterialSlateStyles.h
new file mode 100644
--- /dev/null
+++ b/Engine/Plugins/Runtime/AudioWidgets/Source/AudioWidgets/Public/AudioMaterialSlate/AudioMaterialSlateStyles.h
@@ -0,0 +1,153 @@
+// Copyright Epic Games, Inc. All Rights Reserved.
+
+#pragma once
+
+#include "Materials/MaterialInstanceDynamic.h"
+#include "Materials/MaterialInterface.h"
+#include "Math/Color.h"
+#include "Styling/ISlateStyle.h"
+#include "Styling/SlateWidgetStyle.h"
+#include "Styling/SlateWidgetStyleAsset.h"
+#include "Styling/SlateWidgetStyleContainerBase.h"
+#include "AudioMaterialSlateStyles.generated.h"
+
+
+/**
+ *Base for the appearance of an Audio Material Slates 
+ */
+USTRUCT(BlueprintType)
+struct AUDIOWIDGETS_API FAudioMaterialStyle : public FSlateWidgetStyle
+{
+	GENERATED_USTRUCT_BODY()
+
+	/** Material used to render the Slate */
+	UPROPERTY(EditAnywhere, BlueprintReadWrite, Category = Appearance)
+	UMaterialInterface* Material = nullptr;
+
+public:
+
+	void CreateDynamicMaterial(UObject* InOuter) const;
+	TObjectPtr<UMaterialInstanceDynamic> GetDynamicMaterial() const;
+
+private:
+
+	UPROPERTY()
+	mutable TObjectPtr<UMaterialInstanceDynamic> DynamicMaterial;
+};
+
+/**
+ *Represents the appearance of an Audio Material Button 
+ */
+USTRUCT(BlueprintType)
+struct AUDIOWIDGETS_API FAudioMaterialButtonStyle : public FAudioMaterialStyle
+{
+	GENERATED_USTRUCT_BODY()
+
+	FAudioMaterialButtonStyle();
+	virtual ~FAudioMaterialButtonStyle();
+
+	// FSlateWidgetStyle
+	virtual void GetResources(TArray<const FSlateBrush*>& OutBrushes) const override;
+	static const FName TypeName;
+	virtual const FName GetTypeName() const override { return TypeName; };
+	static const FAudioMaterialButtonStyle& GetDefault();
+
+	UPROPERTY(EditAnywhere, BlueprintReadWrite, Category = Appearance)
+	FLinearColor ButtonMainColor;
+
+	UPROPERTY(EditAnywhere, BlueprintReadWrite, Category = Appearance)
+	FLinearColor ButtonShadowColor;
+
+	UPROPERTY(EditAnywhere, BlueprintReadWrite, Category = Appearance)
+	FLinearColor ButtonAccentColor;	
+	
+	UPROPERTY(EditAnywhere, BlueprintReadWrite, Category = Appearance)
+	FLinearColor ButtonPressedMainColor;	
+	
+	UPROPERTY(EditAnywhere, BlueprintReadWrite, Category = Appearance)
+	FLinearColor ButtonShadowMainColor;	
+	
+	UPROPERTY(EditAnywhere, BlueprintReadWrite, Category = Appearance)
+	FLinearColor ButtonPressedOutlineColor;
+
+};
+
+/**
+ *Represents the appearance of an Audio Material Slider 
+ */
+USTRUCT(BlueprintType)
+struct AUDIOWIDGETS_API FAudioMaterialSliderStyle : public FAudioMaterialStyle
+{
+	GENERATED_USTRUCT_BODY()
+
+	FAudioMaterialSliderStyle();
+	virtual ~FAudioMaterialSliderStyle();
+
+	// FSlateWidgetStyle
+	virtual void GetResources(TArray<const FSlateBrush*>& OutBrushes) const override;
+	static const FName TypeName;
+	virtual const FName GetTypeName() const override { return TypeName; };
+	static const FAudioMaterialSliderStyle& GetDefault();
+
+	UPROPERTY(EditAnywhere, BlueprintReadWrite, Category = Appearance)
+	FLinearColor BarMainColor;	
+	
+	UPROPERTY(EditAnywhere, BlueprintReadWrite, Category = Appearance)
+	FLinearColor BarShadowColor;	
+	
+	UPROPERTY(EditAnywhere, BlueprintReadWrite, Category = Appearance)
+	FLinearColor BarAccentColor;
+
+	UPROPERTY(EditAnywhere, BlueprintReadWrite, Category = Appearance)
+	FLinearColor HandleMainColor;	
+	
+	UPROPERTY(EditAnywhere, BlueprintReadWrite, Category = Appearance)
+	FLinearColor HandleOutlineColor;
+
+};
+
+/**
+ *Represents the appearance of an Audio Material Knob 
+ */
+USTRUCT(BlueprintType)
+struct AUDIOWIDGETS_API FAudioMaterialKnobStyle : public FAudioMaterialStyle
+{
+	GENERATED_USTRUCT_BODY()
+
+	FAudioMaterialKnobStyle();
+	virtual ~FAudioMaterialKnobStyle();
+
+	// FSlateWidgetStyle
+	virtual void GetResources(TArray<const FSlateBrush*>& OutBrushes) const override;
+	static const FName TypeName;
+	virtual const FName GetTypeName() const override { return TypeName; };
+	static const FAudioMaterialKnobStyle& GetDefault();
+
+	UPROPERTY(EditAnywhere, BlueprintReadWrite, Category = Appearance)
+	FLinearColor KnobMainColor;
+
+	UPROPERTY(EditAnywhere, BlueprintReadWrite, Category = Appearance)
+	FLinearColor KnobAccentColor;
+	
+	UPROPERTY(EditAnywhere, BlueprintReadWrite, Category = Appearance)
+	FLinearColor KnobIndicatorColor;
+	
+	UPROPERTY(EditAnywhere, BlueprintReadWrite, Category = Appearance)
+	FLinearColor KnobBarColor;		
+	
+	UPROPERTY(EditAnywhere, BlueprintReadWrite, Category = Appearance)
+	FLinearColor KnobBarShadowColor;	
+
+	UPROPERTY(EditAnywhere, BlueprintReadWrite, Category = Appearance)
+	FLinearColor KnobBarFillMinColor;		
+	
+	UPROPERTY(EditAnywhere, BlueprintReadWrite, Category = Appearance)
+	FLinearColor KnobBarFillMidColor;		
+	
+	UPROPERTY(EditAnywhere, BlueprintReadWrite, Category = Appearance)
+	FLinearColor KnobBarFillMaxColor;	
+	
+	UPROPERTY(EditAnywhere, BlueprintReadWrite, Category = Appearance)
+	FLinearColor KnobBarFillTintColor;		
+	
+};
diff --git a/Engine/Plugins/Runtime/AudioWidgets/Source/AudioWidgets/Public/AudioMaterialSlate/AudioMaterialSlider.h b/Engine/Plugins/Runtime/AudioWidgets/Source/AudioWidgets/Public/AudioMaterialSlate/AudioMaterialSlider.h
new file mode 100644
--- /dev/null
+++ b/Engine/Plugins/Runtime/AudioWidgets/Source/AudioWidgets/Public/AudioMaterialSlate/AudioMaterialSlider.h
@@ -0,0 +1,77 @@
+// Copyright Epic Games, Inc. All Rights Reserved.
+
+#pragma once
+
+#include "AudioMaterialSlateStyles.h"
+#include "Components/Widget.h"
+#include "Delegates/Delegate.h"
+#include "AudioMaterialSlider.generated.h"
+
+class SAudioMaterialSlider;
+
+DECLARE_DYNAMIC_MULTICAST_DELEGATE_OneParam(FOnSliderFloatValueChangedEvent, float, Value);
+
+/**
+ * A simple widget that shows a sliding bar with a handle that allows you to control the value between 0..1.
+ * Slider is rendered by using material instead of texture.
+ *
+ * * No Children
+ */
+UCLASS()
+class AUDIOWIDGETS_API UAudioMaterialSlider : public UWidget
+{
+	GENERATED_BODY()
+
+public:
+
+	/** The slider's style */
+	UPROPERTY(EditAnywhere, BlueprintReadWrite, Category = "Style", meta = (DisplayName = "Style"))
+	FAudioMaterialSliderStyle WidgetStyle;
+
+public:
+
+#if WITH_EDITOR
+	virtual const FText GetPaletteCategory() override;
+#endif
+
+	// UWidget
+	virtual void SynchronizeProperties() override;
+	// End of UWidget
+
+	// UVisual
+	virtual void ReleaseSlateResources(bool bReleaseChildren) override;
+	// End of UVisual
+
+	/** Gets the current value of the slider.*/
+	UFUNCTION(BlueprintPure, Category = "Behavior")
+	float GetValue() const;
+
+	/** Sets the current value of the slider. InValue is Clamped between 0.f - 1.f */
+	UFUNCTION(BlueprintCallable, Category = "Behavior")
+	void SetValue(float InValue);
+
+public:
+
+	/** Called when the value is changed by slider. */
+	UPROPERTY(BlueprintAssignable, Category = "Widget Event")
+	FOnSliderFloatValueChangedEvent OnValueChanged;
+
+protected:
+
+	// UWidget
+	virtual TSharedRef<SWidget> RebuildWidget() override;
+	// End of UWidget
+
+	void HandleOnValueChanged(float InValue);
+
+private:
+
+	/**Default Value of the slider*/
+	UPROPERTY(EditAnywhere, BlueprintSetter = SetValue, BlueprintGetter = GetValue, Category = "Appearance", meta = (UIMin = "0", UIMax = "1"))
+	float Value = 1.f;
+
+private:
+
+	/** Native Slate Widget */
+	TSharedPtr<SAudioMaterialSlider> Slider;
+};
diff --git a/Engine/Plugins/Runtime/AudioWidgets/Source/AudioWidgets/Public/AudioMaterialSlate/SAudioMaterialButton.h b/Engine/Plugins/Runtime/AudioWidgets/Source/AudioWidgets/Public/AudioMaterialSlate/SAudioMaterialButton.h
new file mode 100644
--- /dev/null
+++ b/Engine/Plugins/Runtime/AudioWidgets/Source/AudioWidgets/Public/AudioMaterialSlate/SAudioMaterialButton.h
@@ -0,0 +1,80 @@
+// Copyright Epic Games, Inc. All Rights Reserved.
+
+#pragma once
+
+#include "AudioMaterialSlate/AudioMaterialSlateStyles.h"
+#include "Framework/SlateDelegates.h"
+#include "Styling/ISlateStyle.h"
+#include "Styling/SlateWidgetStyleAsset.h"
+#include "Widgets/SLeafWidget.h"
+
+class UObject;
+struct FAudioMaterialButtonStyle;
+
+/**
+ * A simple slate that renders button in single material and modifies the material on pressed state change.
+ *
+ */
+class AUDIOWIDGETS_API SAudioMaterialButton : public SLeafWidget
+{
+public:
+	SLATE_BEGIN_ARGS(SAudioMaterialButton)
+	{}
+
+	/** The owner object*/
+	SLATE_ARGUMENT(TWeakObjectPtr<UObject>, Owner)
+
+	/**State of the button*/
+	SLATE_ATTRIBUTE(bool, bIsPressedAttribute)
+
+	/** The style used to draw the button. */
+	SLATE_STYLE_ARGUMENT(FAudioMaterialButtonStyle, AudioMaterialButtonStyle)
+
+	/** Called when the button's state changes. */
+	SLATE_EVENT(FOnBooleanValueChanged, OnBooleanValueChanged)
+
+	SLATE_END_ARGS()
+
+	/** Constructs this widget with InArgs */
+	void Construct(const FArguments& InArgs);
+
+	/** Press the button */
+	void SetPressedState(bool InPressedState);
+
+	/** Apply new material to be used to render the Slate.*/
+	void ApplyNewMaterial();
+
+public:
+
+	FOnBooleanValueChanged OnIsPressedStateChanged;
+
+protected:
+
+	// SWidget overrides
+	virtual int32 OnPaint(const FPaintArgs& Args, const FGeometry& AllottedGeometry, const FSlateRect& MyCullingRect, FSlateWindowElementList& OutDrawElements, int32 LayerId, const FWidgetStyle& InWidgetStyle, bool bParentEnabled) const override;
+	virtual FVector2D ComputeDesiredSize(float) const override;
+	virtual FReply OnMouseButtonDown(const FGeometry& MyGeometry, const FPointerEvent& MouseEvent) override;
+	virtual FReply OnMouseButtonUp(const FGeometry& MyGeometry, const FPointerEvent& MouseEvent)override;
+
+private:
+
+	/**Commits new state*/
+	void CommitNewState(bool InPressedState);
+
+private:
+
+	// Holds the owner of the Slate
+	TWeakObjectPtr<UObject> Owner;
+
+	// Holds the style for the Slate
+	const FAudioMaterialButtonStyle* AudioMaterialButtonStyle = nullptr;
+
+	//Current pressed state of this button
+	TAttribute<bool> bIsPressedAttribute = false;
+
+	//Width & height
+	//TODO: Currently hardcoded for testing, make more dynamic.
+	float SlateWidth = 128.f;
+	float SlateHeight = 128.f;
+
+};
diff --git a/Engine/Plugins/Runtime/AudioWidgets/Source/AudioWidgets/Public/AudioMaterialSlate/SAudioMaterialKnob.h b/Engine/Plugins/Runtime/AudioWidgets/Source/AudioWidgets/Public/AudioMaterialSlate/SAudioMaterialKnob.h
new file mode 100644
--- /dev/null
+++ b/Engine/Plugins/Runtime/AudioWidgets/Source/AudioWidgets/Public/AudioMaterialSlate/SAudioMaterialKnob.h
@@ -0,0 +1,89 @@
+// Copyright Epic Games, Inc. All Rights Reserved.
+
+#pragma once
+
+#include "AudioMaterialSlate/AudioMaterialSlateStyles.h"
+#include "Framework/SlateDelegates.h"
+#include "Styling/ISlateStyle.h"
+#include "Styling/SlateWidgetStyleAsset.h"
+#include "Widgets/SLeafWidget.h"
+
+class UObject;
+
+/**
+ * A simple slate that renders a knob in single material and modifies the material on value change.
+ *
+ */
+class AUDIOWIDGETS_API SAudioMaterialKnob : public SLeafWidget
+{
+public:
+	SLATE_BEGIN_ARGS(SAudioMaterialKnob)
+	{}
+
+	/** The owner object*/
+	SLATE_ARGUMENT(TWeakObjectPtr<UObject>, Owner)
+
+	/**Value of the Knob*/
+	SLATE_ATTRIBUTE(float, Value)
+
+	/** The style used to draw the button. */
+	SLATE_STYLE_ARGUMENT(FAudioMaterialKnobStyle, AudioMaterialKnobStyle)
+
+	/** Called when the button's state changes. */
+	SLATE_EVENT(FOnFloatValueChanged, OnFloatValueChanged)
+
+	SLATE_END_ARGS()
+
+	/** Constructs this widget with InArgs */
+	void Construct(const FArguments& InArgs);
+
+	/** Set the Value attribute */
+	void SetValue(float InValueAttribute);
+
+	/** Apply new material to be used to render the Slate.*/
+	void ApplyNewMaterial();
+
+public:
+
+	FOnFloatValueChanged OnValueChanged;
+
+protected:
+
+	// SWidget overrides
+	virtual int32 OnPaint(const FPaintArgs& Args, const FGeometry& AllottedGeometry, const FSlateRect& MyCullingRect, FSlateWindowElementList& OutDrawElements, int32 LayerId, const FWidgetStyle& InWidgetStyle, bool bParentEnabled) const override;
+	virtual FVector2D ComputeDesiredSize(float) const override;
+	virtual FReply OnMouseMove(const FGeometry& MyGeometry, const FPointerEvent& MouseEvent) override;
+	virtual FReply OnMouseButtonDown(const FGeometry& MyGeometry, const FPointerEvent& MouseEvent) override;
+	virtual FReply OnMouseButtonUp(const FGeometry& MyGeometry, const FPointerEvent& MouseEvent)override;
+
+private:
+
+	/**Commits new value*/
+	void CommitValue(float NewValue);
+
+private:
+
+	// Holds the owner of the Slate
+	TWeakObjectPtr<UObject> Owner;
+
+	// Holds the style for the Slate
+	const FAudioMaterialKnobStyle* AudioMaterialKnobStyle = nullptr;
+
+	//Current Value
+	TAttribute<float> ValueAttribute = 1.0;
+
+	// The position of the mouse when it pushed down and started rotating the knob
+	FVector2D MouseDownPosition;
+
+	// the value when the mouse was pushed down
+	float MouseDownValue = 0.f;
+
+	// the max pixels to go to min or max value (clamped to 0 or 1) in one drag period
+	int32 PixelDelta = 50;
+
+	//Width & height
+	//TODO: Currently hardcoded for testing, make more dynamic.
+	float SlateWidth = 256.f;
+	float SlateHeight = 256.f;
+
+};
diff --git a/Engine/Plugins/Runtime/AudioWidgets/Source/AudioWidgets/Public/AudioMaterialSlate/SAudioMaterialSlider.h b/Engine/Plugins/Runtime/AudioWidgets/Source/AudioWidgets/Public/AudioMaterialSlate/SAudioMaterialSlider.h
new file mode 100644
--- /dev/null
+++ b/Engine/Plugins/Runtime/AudioWidgets/Source/AudioWidgets/Public/AudioMaterialSlate/SAudioMaterialSlider.h
@@ -0,0 +1,96 @@
+// Copyright Epic Games, Inc. All Rights Reserved.
+
+#pragma once
+
+#include "AudioMaterialSlate/AudioMaterialSlateStyles.h"
+#include "Framework/SlateDelegates.h"
+#include "Styling/ISlateStyle.h"
+#include "Styling/SlateWidgetStyleAsset.h"
+#include "Widgets/SLeafWidget.h"
+
+class UWidget;
+
+/**
+ * A simple slate that renders slider in single material and modifies the material on value change.
+ *
+ */
+class AUDIOWIDGETS_API SAudioMaterialSlider : public SLeafWidget
+{
+public:
+	SLATE_BEGIN_ARGS(SAudioMaterialSlider)
+	{}
+
+	/** The owner object*/
+	SLATE_ARGUMENT(TWeakObjectPtr<UObject>, Owner)
+
+	/** The style used to draw the slider. */
+	SLATE_STYLE_ARGUMENT(FAudioMaterialSliderStyle, AudioMaterialSliderStyle)
+
+	/** A value that drives where the slider handle appears. Value is clamped between 0 and 1. */
+	SLATE_ATTRIBUTE(float, ValueAttribute)
+
+	/** Called when the value is changed by the slider. */
+	SLATE_EVENT(FOnFloatValueChanged, OnValueChanged)
+
+	SLATE_END_ARGS()
+
+	/**
+	* Construct the widget.
+	*/
+	void Construct(const FArguments& InArgs);
+
+	// SWidget overrides
+	virtual int32 OnPaint(const FPaintArgs& Args, const FGeometry& AllottedGeometry, const FSlateRect& MyCullingRect, FSlateWindowElementList& OutDrawElements, int32 LayerId, const FWidgetStyle& InWidgetStyle, bool bParentEnabled) const override;
+	virtual FVector2D ComputeDesiredSize(float) const override;
+
+	/** Set the Value attribute */
+	void SetValue(TAttribute<float> InValueAttribute);
+
+	/** Apply new material to be used to render the Slate.*/
+	void ApplyNewMaterial();
+
+	// SWidget overrides
+	virtual FReply OnMouseButtonDown(const FGeometry& MyGeometry, const FPointerEvent& MouseEvent) override;
+	virtual FReply OnMouseButtonUp(const FGeometry& MyGeometry, const FPointerEvent& MouseEvent) override;
+	virtual FReply OnMouseMove(const FGeometry& MyGeometry, const FPointerEvent& MouseEvent) override;
+
+public:
+
+	// Holds a delegate that is executed when the slider's value changed.
+	FOnFloatValueChanged OnValueChanged;	
+
+private:
+
+	/**
+	* Commits the specified slider value.
+	*
+	* @param NewValue The value to commit.
+	*/
+	void CommitValue(float NewValue);
+
+	/**
+	* Calculates the new value based on the given absolute coordinates.
+	*
+	* @param MyGeometry The slider's geometry.
+	* @param AbsolutePosition The absolute position of the slider.
+	* @return The new value.
+	*/
+	FVector2D PositionToValue(const FGeometry& MyGeometry, const FVector2D& AbsolutePosition);
+
+private:
+
+	// Holds the owner of the Slate
+	TWeakObjectPtr<UObject> Owner;
+
+	// Holds the style for the Slate
+	const FAudioMaterialSliderStyle* AudioMaterialSliderStyle = nullptr;
+
+	//Holds the current value
+	TAttribute<float> ValueAttribute = 0.f;
+
+	//Width & height
+	//TODO: Currently hardcoded for testing, make more dynamic.
+	float SlateWidth = 128.f;
+	float SlateHeight = 512.f;
+
+};
