Added default context menu for spectrogram widget. This can be overridden or extended if the context in which the widget is used requires it.

diff --git a/Engine/Plugins/Runtime/AudioWidgets/Source/AudioWidgets/Private/SAudioSpectrogram.cpp b/Engine/Plugins/Runtime/AudioWidgets/Source/AudioWidgets/Private/SAudioSpectrogram.cpp
--- a/Engine/Plugins/Runtime/AudioWidgets/Source/AudioWidgets/Private/SAudioSpectrogram.cpp
+++ b/Engine/Plugins/Runtime/AudioWidgets/Source/AudioWidgets/Private/SAudioSpectrogram.cpp
@@ -2,6 +2,14 @@
 
 #include "SAudioSpectrogram.h"
 
+#include "Framework/Application/SlateApplication.h"
+#include "Framework/MultiBox/MultiBoxBuilder.h"
+#include "Layout/WidgetPath.h"
+
+#define LOCTEXT_NAMESPACE "SAudioSpectrogram"
+
+FName SAudioSpectrogram::ContextMenuExtensionHook("SpectrogramDisplayOptions");
+
 void SAudioSpectrogram::Construct(const FArguments& InArgs)
 {
 	ViewMinFrequency = InArgs._ViewMinFrequency;
@@ -12,6 +20,8 @@ void SAudioSpectrogram::Construct(const FArguments& InArgs)
 	FrequencyAxisScale = InArgs._FrequencyAxisScale;
 	FrequencyAxisPixelBucketMode = InArgs._FrequencyAxisPixelBucketMode;
 	Orientation = InArgs._Orientation;
+	bAllowContextMenu = InArgs._AllowContextMenu;
+	OnContextMenuOpening = InArgs._OnContextMenuOpening;
 
 	SpectrogramViewport = MakeShareable(new FAudioSpectrogramViewport());
 }
@@ -21,6 +31,69 @@ void SAudioSpectrogram::AddFrame(const FAudioSpectrogramFrameData& SpectrogramFr
 	SpectrogramViewport->AddFrame(SpectrogramFrameData);
 }
 
+TSharedRef<const FExtensionBase> SAudioSpectrogram::AddContextMenuExtension(EExtensionHook::Position HookPosition, const TSharedPtr<FUICommandList>& CommandList, const FMenuExtensionDelegate& MenuExtensionDelegate)
+{
+	if (!ContextMenuExtender.IsValid())
+	{
+		ContextMenuExtender = MakeShared<FExtender>();
+	}
+
+	return ContextMenuExtender->AddMenuExtension(ContextMenuExtensionHook, HookPosition, CommandList, MenuExtensionDelegate);
+}
+
+void SAudioSpectrogram::RemoveContextMenuExtension(const TSharedRef<const FExtensionBase>& Extension)
+{
+	if (ensure(ContextMenuExtender.IsValid()))
+	{
+		ContextMenuExtender->RemoveExtension(Extension);
+	}
+}
+
+FReply SAudioSpectrogram::OnMouseButtonDown(const FGeometry& InMyGeometry, const FPointerEvent& InMouseEvent)
+{
+	if (!HasMouseCapture())
+	{
+		if (InMouseEvent.GetEffectingButton() == EKeys::RightMouseButton)
+		{
+			// Right clicking to summon context menu, but we'll do that on mouse-up.
+			return FReply::Handled().CaptureMouse(AsShared()).SetUserFocus(AsShared(), EFocusCause::Mouse);
+		}
+	}
+
+	return SCompoundWidget::OnMouseButtonDown(InMyGeometry, InMouseEvent);
+}
+
+FReply SAudioSpectrogram::OnMouseButtonUp(const FGeometry& InMyGeometry, const FPointerEvent& InMouseEvent)
+{
+	// The mouse must have been captured by mouse down before we'll process mouse ups
+	if (HasMouseCapture())
+	{
+		if (InMouseEvent.GetEffectingButton() == EKeys::RightMouseButton)
+		{
+			if (InMyGeometry.IsUnderLocation(InMouseEvent.GetScreenSpacePosition()) && bAllowContextMenu.Get())
+			{
+				TSharedPtr<SWidget> ContextMenu = OnContextMenuOpening.IsBound() ? OnContextMenuOpening.Execute() : BuildDefaultContextMenu();
+
+				if (ContextMenu.IsValid())
+				{
+					const FWidgetPath WidgetPath = (InMouseEvent.GetEventPath() != nullptr) ? *InMouseEvent.GetEventPath() : FWidgetPath();
+
+					FSlateApplication::Get().PushMenu(
+						AsShared(),
+						WidgetPath,
+						ContextMenu.ToSharedRef(),
+						InMouseEvent.GetScreenSpacePosition(),
+						FPopupTransitionEffect::ESlideDirection::ContextMenu);
+				}
+			}
+
+			return FReply::Handled().ReleaseMouseCapture();
+		}
+	}
+
+	return SCompoundWidget::OnMouseButtonUp(InMyGeometry, InMouseEvent);
+}
+
 int32 SAudioSpectrogram::OnPaint(const FPaintArgs& Args, const FGeometry& AllottedGeometry, const FSlateRect& MyCullingRect, FSlateWindowElementList& OutDrawElements, int32 LayerId, const FWidgetStyle& InWidgetStyle, bool bParentEnabled) const
 {
 	// Get post DPI-scaled sizing:
@@ -61,3 +134,153 @@ int32 SAudioSpectrogram::OnPaint(const FPaintArgs& Args, const FGeometry& Allott
 
 	return LayerId + 1;
 }
+
+TSharedRef<SWidget> SAudioSpectrogram::BuildDefaultContextMenu()
+{
+	constexpr bool bShouldCloseWindowAfterMenuSelection = true;
+	FMenuBuilder MenuBuilder(bShouldCloseWindowAfterMenuSelection, nullptr, ContextMenuExtender);
+
+	MenuBuilder.BeginSection(ContextMenuExtensionHook, LOCTEXT("DisplayOptions", "Display Options"));
+
+	if (!FrequencyAxisPixelBucketMode.IsBound())
+	{
+		MenuBuilder.AddSubMenu(
+			LOCTEXT("FrequencyAxisPixelBucketMode", "Pixel Plot Mode"),
+			FText(),
+			FNewMenuDelegate::CreateSP(this, &SAudioSpectrogram::BuildFrequencyAxisPixelBucketModeSubMenu));
+	}
+
+	if (!FrequencyAxisScale.IsBound())
+	{
+		MenuBuilder.AddSubMenu(
+			LOCTEXT("FrequencyAxisScale", "Frequency Scale"),
+			FText(),
+			FNewMenuDelegate::CreateSP(this, &SAudioSpectrogram::BuildFrequencyAxisScaleSubMenu));
+	}
+
+	if (!ColorMap.IsBound())
+	{
+		MenuBuilder.AddSubMenu(
+			LOCTEXT("ColorMap", "Color Map"),
+			FText(),
+			FNewMenuDelegate::CreateSP(this, &SAudioSpectrogram::BuildColorMapSubMenu));
+	}
+
+	if (!Orientation.IsBound())
+	{
+		MenuBuilder.AddSubMenu(
+			LOCTEXT("Orientation", "Orientation"),
+			FText(),
+			FNewMenuDelegate::CreateSP(this, &SAudioSpectrogram::BuildOrientationSubMenu));
+	}
+
+	MenuBuilder.EndSection();
+
+	return MenuBuilder.MakeWidget();
+}
+
+void SAudioSpectrogram::BuildColorMapSubMenu(FMenuBuilder& SubMenu)
+{
+	const UEnum* EnumClass = StaticEnum<EAudioColorMap>();
+	const int32 NumEnumValues = EnumClass->NumEnums() - 1; // Exclude 'MAX' enum value.
+	for (int32 Index = 0; Index < NumEnumValues; Index++)
+	{
+		const auto EnumValue = static_cast<EAudioColorMap>(EnumClass->GetValueByIndex(Index));
+
+		SubMenu.AddMenuEntry(
+			EnumClass->GetDisplayNameTextByIndex(Index),
+#if WITH_EDITOR
+			EnumClass->GetToolTipTextByIndex(Index),
+#else
+			FText(),
+#endif
+			FSlateIcon(),
+			FUIAction(
+				FExecuteAction::CreateSP(this, &SAudioSpectrogram::SetColorMap, EnumValue),
+				FCanExecuteAction(),
+				FIsActionChecked::CreateSPLambda(this, [this, EnumValue]() { return (ColorMap.Get() == EnumValue); })
+			),
+			NAME_None,
+			EUserInterfaceActionType::ToggleButton);
+	}
+}
+
+void SAudioSpectrogram::BuildFrequencyAxisScaleSubMenu(FMenuBuilder& SubMenu)
+{
+	const UEnum* EnumClass = StaticEnum<EAudioSpectrogramFrequencyAxisScale>();
+	const int32 NumEnumValues = EnumClass->NumEnums() - 1; // Exclude 'MAX' enum value.
+	for (int32 Index = 0; Index < NumEnumValues; Index++)
+	{
+		const auto EnumValue = static_cast<EAudioSpectrogramFrequencyAxisScale>(EnumClass->GetValueByIndex(Index));
+
+		SubMenu.AddMenuEntry(
+			EnumClass->GetDisplayNameTextByIndex(Index),
+#if WITH_EDITOR
+			EnumClass->GetToolTipTextByIndex(Index),
+#else
+			FText(),
+#endif
+			FSlateIcon(),
+			FUIAction(
+				FExecuteAction::CreateSP(this, &SAudioSpectrogram::SetFrequencyAxisScale, EnumValue),
+				FCanExecuteAction(),
+				FIsActionChecked::CreateSPLambda(this, [this, EnumValue]() { return (FrequencyAxisScale.Get() == EnumValue); })
+			),
+			NAME_None,
+			EUserInterfaceActionType::ToggleButton);
+	}
+}
+
+void SAudioSpectrogram::BuildFrequencyAxisPixelBucketModeSubMenu(FMenuBuilder& SubMenu)
+{
+	const UEnum* EnumClass = StaticEnum<EAudioSpectrogramFrequencyAxisPixelBucketMode>();
+	const int32 NumEnumValues = EnumClass->NumEnums() - 1; // Exclude 'MAX' enum value.
+	for (int32 Index = 0; Index < NumEnumValues; Index++)
+	{
+		const auto EnumValue = static_cast<EAudioSpectrogramFrequencyAxisPixelBucketMode>(EnumClass->GetValueByIndex(Index));
+
+		SubMenu.AddMenuEntry(
+			EnumClass->GetDisplayNameTextByIndex(Index),
+#if WITH_EDITOR
+			EnumClass->GetToolTipTextByIndex(Index),
+#else
+			FText(),
+#endif
+			FSlateIcon(),
+			FUIAction(
+				FExecuteAction::CreateSP(this, &SAudioSpectrogram::SetFrequencyAxisPixelBucketMode, EnumValue),
+				FCanExecuteAction(),
+				FIsActionChecked::CreateSPLambda(this, [this, EnumValue]() { return (FrequencyAxisPixelBucketMode.Get() == EnumValue); })
+			),
+			NAME_None,
+			EUserInterfaceActionType::ToggleButton);
+	}
+}
+
+void SAudioSpectrogram::BuildOrientationSubMenu(FMenuBuilder& SubMenu)
+{
+	const UEnum* EnumClass = StaticEnum<EOrientation>();
+	const int32 NumEnumValues = EnumClass->NumEnums() - 1; // Exclude 'MAX' enum value.
+	for (int32 Index = 0; Index < NumEnumValues; Index++)
+	{
+		const auto EnumValue = static_cast<EOrientation>(EnumClass->GetValueByIndex(Index));
+
+		SubMenu.AddMenuEntry(
+			EnumClass->GetDisplayNameTextByIndex(Index),
+#if WITH_EDITOR
+			EnumClass->GetToolTipTextByIndex(Index),
+#else
+			FText(),
+#endif
+			FSlateIcon(),
+			FUIAction(
+				FExecuteAction::CreateSP(this, &SAudioSpectrogram::SetOrientation, EnumValue),
+				FCanExecuteAction(),
+				FIsActionChecked::CreateSPLambda(this, [this, EnumValue]() { return (Orientation.Get() == EnumValue); })
+			),
+			NAME_None,
+			EUserInterfaceActionType::ToggleButton);
+	}
+}
+
+#undef LOCTEXT_NAMESPACE
diff --git a/Engine/Plugins/Runtime/AudioWidgets/Source/AudioWidgets/Public/SAudioSpectrogram.h b/Engine/Plugins/Runtime/AudioWidgets/Source/AudioWidgets/Public/SAudioSpectrogram.h
--- a/Engine/Plugins/Runtime/AudioWidgets/Source/AudioWidgets/Public/SAudioSpectrogram.h
+++ b/Engine/Plugins/Runtime/AudioWidgets/Source/AudioWidgets/Public/SAudioSpectrogram.h
@@ -4,6 +4,8 @@
 
 #include "AudioColorMapper.h"
 #include "AudioSpectrogramViewport.h"
+#include "Framework/MultiBox/MultiBoxExtender.h"
+#include "Framework/SlateDelegates.h"
 #include "Widgets/SCompoundWidget.h"
 
 /**
@@ -21,6 +23,7 @@ public:
 		, _FrequencyAxisScale(EAudioSpectrogramFrequencyAxisScale::Logarithmic)
 		, _FrequencyAxisPixelBucketMode(EAudioSpectrogramFrequencyAxisPixelBucketMode::Average)
 		, _Orientation(EOrientation::Orient_Horizontal)
+		, _AllowContextMenu(true)
 	{}
 		SLATE_ATTRIBUTE(float, ViewMinFrequency)
 		SLATE_ATTRIBUTE(float, ViewMaxFrequency)
@@ -30,6 +33,8 @@ public:
 		SLATE_ATTRIBUTE(EAudioSpectrogramFrequencyAxisScale, FrequencyAxisScale)
 		SLATE_ATTRIBUTE(EAudioSpectrogramFrequencyAxisPixelBucketMode, FrequencyAxisPixelBucketMode)
 		SLATE_ATTRIBUTE(EOrientation, Orientation)
+		SLATE_ATTRIBUTE(bool, AllowContextMenu)
+		SLATE_EVENT(FOnContextMenuOpening, OnContextMenuOpening)
 	SLATE_END_ARGS()
 
 	/** Constructs this widget with InArgs */
@@ -46,12 +51,30 @@ public:
 	void SetFrequencyAxisScale(const EAudioSpectrogramFrequencyAxisScale InFrequencyAxisScale) { FrequencyAxisScale = InFrequencyAxisScale; }
 	void SetFrequencyAxisPixelBucketMode(const EAudioSpectrogramFrequencyAxisPixelBucketMode InFrequencyAxisPixelBucketMode) { FrequencyAxisPixelBucketMode = InFrequencyAxisPixelBucketMode; }
 	void SetOrientation(const EOrientation InOrientation) { Orientation = InOrientation; }
+	void SetAllowContextMenu(bool bInAllowContextMenu) { bAllowContextMenu = bInAllowContextMenu; }
+
+	TSharedRef<const FExtensionBase> AddContextMenuExtension(EExtensionHook::Position HookPosition, const TSharedPtr<FUICommandList>& CommandList, const FMenuExtensionDelegate& MenuExtensionDelegate);
+	void RemoveContextMenuExtension(const TSharedRef<const FExtensionBase>& Extension);
+
+	// Begin SWidget overrides.
+	virtual FReply OnMouseButtonDown(const FGeometry& InMyGeometry, const FPointerEvent& InMouseEvent) override;
+	virtual FReply OnMouseButtonUp(const FGeometry& InMyGeometry, const FPointerEvent& InMouseEvent) override;
+	// End SWidget overrides.
 
 private:
 	// Begin SWidget overrides.
 	virtual int32 OnPaint(const FPaintArgs& Args, const FGeometry& AllottedGeometry, const FSlateRect& MyCullingRect, FSlateWindowElementList& OutDrawElements, int32 LayerId, const FWidgetStyle& InWidgetStyle, bool bParentEnabled) const override;
 	// End SWidget overrides.
 
+	TSharedRef<SWidget> BuildDefaultContextMenu();
+	void BuildColorMapSubMenu(FMenuBuilder& SubMenu);
+	void BuildFrequencyAxisScaleSubMenu(FMenuBuilder& SubMenu);
+	void BuildFrequencyAxisPixelBucketModeSubMenu(FMenuBuilder& SubMenu);
+	void BuildOrientationSubMenu(FMenuBuilder& SubMenu);
+
+	static FName ContextMenuExtensionHook;
+	TSharedPtr<FExtender> ContextMenuExtender;
+
 	TAttribute<float> ViewMinFrequency;
 	TAttribute<float> ViewMaxFrequency;
 	TAttribute<float> ColorMapMinSoundLevel;
@@ -60,6 +83,8 @@ private:
 	TAttribute<EAudioSpectrogramFrequencyAxisScale> FrequencyAxisScale;
 	TAttribute<EAudioSpectrogramFrequencyAxisPixelBucketMode> FrequencyAxisPixelBucketMode;
 	TAttribute<EOrientation> Orientation;
+	TAttribute<bool> bAllowContextMenu;
+	FOnContextMenuOpening OnContextMenuOpening;
 
 	TSharedPtr<FAudioSpectrogramViewport> SpectrogramViewport;
 };
