AudioMaterialMeter first pass

#virtualized

diff --git a/Engine/Plugins/Runtime/AudioWidgets/Source/AudioWidgets/Private/AudioMaterialSlate/AudioMaterialMeter.cpp b/Engine/Plugins/Runtime/AudioWidgets/Source/AudioWidgets/Private/AudioMaterialSlate/AudioMaterialMeter.cpp
new file mode 100644
--- /dev/null
+++ b/Engine/Plugins/Runtime/AudioWidgets/Source/AudioWidgets/Private/AudioMaterialSlate/AudioMaterialMeter.cpp
@@ -0,0 +1,72 @@
+// Copyright Epic Games, Inc. All Rights Reserved.
+
+
+#include "AudioMaterialSlate/AudioMaterialMeter.h"
+#include "AudioMaterialSlate/AudioMaterialSlateStyles.h"
+#include "AudioMaterialSlate/SAudioMaterialMeter.h"
+#include "Widgets/SWeakWidget.h"
+#include "Components/AudioComponent.h"
+
+#define LOCTEXT_NAMESPACE "AudioWidgets"
+
+#if WITH_EDITOR
+const FText UAudioMaterialMeter::GetPaletteCategory()
+{
+	return LOCTEXT("PaletteCategory", "AudioMaterial");
+}
+#endif
+
+void UAudioMaterialMeter::SynchronizeProperties()
+{
+	Super::SynchronizeProperties();
+
+	if (!Meter.IsValid())
+	{
+		return;
+	}
+
+	Meter->ApplyNewMaterial();
+	Meter->SetValue(MeterValue);
+}
+
+void UAudioMaterialMeter::ReleaseSlateResources(bool bReleaseChildren)
+{
+	Meter.Reset();
+}
+
+float UAudioMaterialMeter::GetValue() const
+{
+	return MeterValue;
+}
+
+void UAudioMaterialMeter::SetValue(float InValue)
+{
+	if (Meter.IsValid())
+	{
+		Meter->SetValue(InValue);
+	}
+
+	if (MeterValue != InValue)
+	{
+		MeterValue = InValue;
+		HandleOnValueChanged(InValue);
+	}
+}
+
+TSharedRef<SWidget> UAudioMaterialMeter::RebuildWidget()
+{
+	Meter = SNew(SAudioMaterialMeter)
+		.Owner(this)
+		.AudioMaterialMeterStyle(&WidgetStyle)
+		.ValueAttribute(MeterValue)
+		.OnValueChanged(BIND_UOBJECT_DELEGATE(FOnFloatValueChanged, HandleOnValueChanged));
+
+	return Meter.ToSharedRef();
+}
+
+void UAudioMaterialMeter::HandleOnValueChanged(float InValue)
+{
+	OnValueChanged.Broadcast(InValue);
+}
+
+#undef LOCTEXT_NAMESPACE
diff --git a/Engine/Plugins/Runtime/AudioWidgets/Source/AudioWidgets/Private/AudioMaterialSlate/AudioMaterialSlateStyles.cpp b/Engine/Plugins/Runtime/AudioWidgets/Source/AudioWidgets/Private/AudioMaterialSlate/AudioMaterialSlateStyles.cpp
--- a/Engine/Plugins/Runtime/AudioWidgets/Source/AudioWidgets/Private/AudioMaterialSlate/AudioMaterialSlateStyles.cpp
+++ b/Engine/Plugins/Runtime/AudioWidgets/Source/AudioWidgets/Private/AudioMaterialSlate/AudioMaterialSlateStyles.cpp
@@ -94,3 +94,28 @@ void FAudioMaterialKnobStyle::GetResources(TArray<const FSlateBrush*>& OutBrushe
 {
 	// Add any brush resources here so that Slate can correctly atlas and reference them
 }
+
+FAudioMaterialMeterStyle::FAudioMaterialMeterStyle()
+:MeterFillMinColor(FLinearColor(1.f, 1.f, 1.f, 1.f)),
+MeterFillMidColor(FLinearColor(1.f, 1.f, 1.f, 1.f)),
+MeterFillMaxColor(FLinearColor(1.f, 1.f, 1.f, 1.f)),
+MeterOffFillColor(FLinearColor(0.5f, 0.5f, 0.5f, 1.f))
+{
+}
+
+FAudioMaterialMeterStyle::~FAudioMaterialMeterStyle()
+{
+}
+
+const FName FAudioMaterialMeterStyle::TypeName(TEXT("FAudioMaterialKnobStyle"));
+
+const FAudioMaterialMeterStyle& FAudioMaterialMeterStyle::GetDefault()
+{
+	static FAudioMaterialMeterStyle Default;
+	return Default;
+}
+
+void FAudioMaterialMeterStyle::GetResources(TArray<const FSlateBrush*>& OutBrushes) const
+{
+	// Add any brush resources here so that Slate can correctly atlas and reference them
+}
diff --git a/Engine/Plugins/Runtime/AudioWidgets/Source/AudioWidgets/Private/AudioMaterialSlate/SAudioMaterialMeter.cpp b/Engine/Plugins/Runtime/AudioWidgets/Source/AudioWidgets/Private/AudioMaterialSlate/SAudioMaterialMeter.cpp
new file mode 100644
--- /dev/null
+++ b/Engine/Plugins/Runtime/AudioWidgets/Source/AudioWidgets/Private/AudioMaterialSlate/SAudioMaterialMeter.cpp
@@ -0,0 +1,84 @@
+// Copyright Epic Games, Inc. All Rights Reserved.
+
+
+#include "AudioMaterialSlate/SAudioMaterialMeter.h"
+#include "AudioMaterialSlate/AudioMaterialMeter.h"
+#include "Components/AudioComponent.h"
+#include "SlateOptMacros.h"
+#include "Styling/SlateBrush.h"
+
+
+BEGIN_SLATE_FUNCTION_BUILD_OPTIMIZATION
+
+void SAudioMaterialMeter::Construct(const FArguments& InArgs)
+{
+	Owner = InArgs._Owner;
+	ValueAttribute = InArgs._ValueAttribute;
+	OnValueChanged = InArgs._OnValueChanged;
+	AudioMaterialMeterStyle = InArgs._AudioMaterialMeterStyle;
+
+	ApplyNewMaterial();
+}
+
+END_SLATE_FUNCTION_BUILD_OPTIMIZATION
+
+int32 SAudioMaterialMeter::OnPaint(const FPaintArgs& Args, const FGeometry& AllottedGeometry, const FSlateRect& MyCullingRect, FSlateWindowElementList& OutDrawElements, int32 LayerId, const FWidgetStyle& InWidgetStyle, bool bParentEnabled) const
+{
+	if (AudioMaterialMeterStyle)
+	{
+		UMaterialInstanceDynamic* DynamicMaterial = AudioMaterialMeterStyle->GetDynamicMaterial();
+		if (IsValid(DynamicMaterial))
+		{
+
+			DynamicMaterial->SetVectorParameterValue(FName("A (V3)"), AudioMaterialMeterStyle->MeterFillMinColor);
+			DynamicMaterial->SetVectorParameterValue(FName("B (V3)"), AudioMaterialMeterStyle->MeterFillMidColor);
+			DynamicMaterial->SetVectorParameterValue(FName("C (V3)"), AudioMaterialMeterStyle->MeterFillMaxColor);
+			DynamicMaterial->SetVectorParameterValue(FName("OffColor"), AudioMaterialMeterStyle->MeterOffFillColor);
+			DynamicMaterial->SetVectorParameterValue(FName("DotsOffColor"), AudioMaterialMeterStyle->MeterOffFillColor);
+
+			const float Value = ValueAttribute.Get();
+			DynamicMaterial->SetScalarParameterValue(FName("VALUE"), FMath::Clamp(Value, 0.f, 1.f));
+		}
+
+		const bool bEnabled = ShouldBeEnabled(bParentEnabled);
+		const ESlateDrawEffect DrawEffects = bEnabled ? ESlateDrawEffect::None : ESlateDrawEffect::DisabledEffect;
+			
+		const FLinearColor FinalColorAndOpacity(InWidgetStyle.GetColorAndOpacityTint());
+
+		FSlateBrush Brush;
+		Brush.SetResourceObject(DynamicMaterial);
+		FSlateDrawElement::MakeBox(OutDrawElements, LayerId++, AllottedGeometry.ToPaintGeometry(), &Brush, DrawEffects, FinalColorAndOpacity);
+	}
+
+	return LayerId;
+}
+
+FVector2D SAudioMaterialMeter::ComputeDesiredSize(float) const
+{
+	FVector2D Vector = FVector2D(SlateWidth, SlateHeight);
+	return Vector;
+}
+
+void SAudioMaterialMeter::SetValue(float InValueAttribute)
+{
+	CommitValue(InValueAttribute);
+}
+
+void SAudioMaterialMeter::ApplyNewMaterial()
+{
+	if (AudioMaterialMeterStyle)
+	{
+		AudioMaterialMeterStyle->CreateDynamicMaterial(Owner.Get());
+	}
+}
+
+void SAudioMaterialMeter::CommitValue(float NewValue)
+{
+	float Val = FMath::Clamp(NewValue, 0.f , 1.f);
+	if (!ValueAttribute.IsBound())
+	{
+		ValueAttribute.Set(Val);
+	}
+
+	OnValueChanged.ExecuteIfBound(Val);
+}
diff --git a/Engine/Plugins/Runtime/AudioWidgets/Source/AudioWidgets/Public/AudioMaterialSlate/AudioMaterialMeter.h b/Engine/Plugins/Runtime/AudioWidgets/Source/AudioWidgets/Public/AudioMaterialSlate/AudioMaterialMeter.h
new file mode 100644
--- /dev/null
+++ b/Engine/Plugins/Runtime/AudioWidgets/Source/AudioWidgets/Public/AudioMaterialSlate/AudioMaterialMeter.h
@@ -0,0 +1,75 @@
+// Copyright Epic Games, Inc. All Rights Reserved.
+
+#pragma once
+
+#include "AudioMaterialSlateStyles.h"
+#include "Components/Widget.h"
+#include "Delegates/Delegate.h"
+#include "AudioMaterialMeter.generated.h"
+
+class SAudioMaterialMeter;
+
+DECLARE_DYNAMIC_MULTICAST_DELEGATE_OneParam(FOnMeterValueChangedEvent, float, Value);
+
+/**
+ * Meter is rendered by using material instead of texture.
+ *
+ * * No Children
+ */
+UCLASS()
+class AUDIOWIDGETS_API UAudioMaterialMeter : public UWidget
+{
+	GENERATED_BODY()
+
+public:
+
+	/** The meter's style */
+	UPROPERTY(EditAnywhere, BlueprintReadWrite, Category = "Style", meta = (DisplayName = "Style"))
+	FAudioMaterialMeterStyle WidgetStyle;
+
+public:
+
+#if WITH_EDITOR
+	virtual const FText GetPaletteCategory() override;
+#endif
+
+	// UWidget
+	virtual void SynchronizeProperties() override;
+	// End of UWidget
+
+	// UVisual
+	virtual void ReleaseSlateResources(bool bReleaseChildren) override;
+	// End of UVisual
+
+	/** Get the current value of the meter.*/
+	UFUNCTION(BlueprintPure, Category = "Behavior")
+	float GetValue() const;
+
+	/** Set the current value of the Meter.*/
+	UFUNCTION(BlueprintCallable, Category = "Behavior")
+	void SetValue(float InValue);
+
+public:
+
+	/** Called when the meter value changes. */
+	UPROPERTY(BlueprintAssignable, Category = "Widget Event")
+	FOnMeterValueChangedEvent OnValueChanged;
+
+protected:
+
+	// UWidget
+	virtual TSharedRef<SWidget> RebuildWidget() override;
+	// End of UWidget
+
+	void HandleOnValueChanged(float InValue);
+
+private:
+
+	/**Current Value of the Meter*/
+	UPROPERTY(EditAnywhere, BlueprintSetter = SetValue, BlueprintGetter = GetValue, Category = "Appearance", meta = (UIMin = "0", UIMax = "1"))
+	float MeterValue = 1.f;
+
+	/** Native Slate Widget */
+	TSharedPtr<SAudioMaterialMeter> Meter;
+
+};
diff --git a/Engine/Plugins/Runtime/AudioWidgets/Source/AudioWidgets/Public/AudioMaterialSlate/AudioMaterialSlateStyles.h b/Engine/Plugins/Runtime/AudioWidgets/Source/AudioWidgets/Public/AudioMaterialSlate/AudioMaterialSlateStyles.h
--- a/Engine/Plugins/Runtime/AudioWidgets/Source/AudioWidgets/Public/AudioMaterialSlate/AudioMaterialSlateStyles.h
+++ b/Engine/Plugins/Runtime/AudioWidgets/Source/AudioWidgets/Public/AudioMaterialSlate/AudioMaterialSlateStyles.h
@@ -151,3 +151,34 @@ struct AUDIOWIDGETS_API FAudioMaterialKnobStyle : public FAudioMaterialStyle
 	FLinearColor KnobBarFillTintColor;		
 	
 };
+
+/**
+ *Represents the appearance of an Audio Material Meter
+ */
+USTRUCT(BlueprintType)
+struct AUDIOWIDGETS_API FAudioMaterialMeterStyle : public FAudioMaterialStyle
+{
+	GENERATED_USTRUCT_BODY()
+
+	FAudioMaterialMeterStyle();
+	virtual ~FAudioMaterialMeterStyle();
+
+	// FSlateWidgetStyle
+	virtual void GetResources(TArray<const FSlateBrush*>& OutBrushes) const override;
+	static const FName TypeName;
+	virtual const FName GetTypeName() const override { return TypeName; };
+	static const FAudioMaterialMeterStyle& GetDefault();
+
+	UPROPERTY(EditAnywhere, BlueprintReadWrite, Category = Appearance)
+	FLinearColor MeterFillMinColor;
+
+	UPROPERTY(EditAnywhere, BlueprintReadWrite, Category = Appearance)
+	FLinearColor MeterFillMidColor;
+
+	UPROPERTY(EditAnywhere, BlueprintReadWrite, Category = Appearance)
+	FLinearColor MeterFillMaxColor;	
+	
+	UPROPERTY(EditAnywhere, BlueprintReadWrite, Category = Appearance)
+	FLinearColor MeterOffFillColor;
+	
+};
diff --git a/Engine/Plugins/Runtime/AudioWidgets/Source/AudioWidgets/Public/AudioMaterialSlate/SAudioMaterialMeter.h b/Engine/Plugins/Runtime/AudioWidgets/Source/AudioWidgets/Public/AudioMaterialSlate/SAudioMaterialMeter.h
new file mode 100644
--- /dev/null
+++ b/Engine/Plugins/Runtime/AudioWidgets/Source/AudioWidgets/Public/AudioMaterialSlate/SAudioMaterialMeter.h
@@ -0,0 +1,82 @@
+// Copyright Epic Games, Inc. All Rights Reserved.
+
+#pragma once
+
+#include "AudioMaterialSlate/AudioMaterialSlateStyles.h"
+#include "Framework/SlateDelegates.h"
+#include "Styling/ISlateStyle.h"
+#include "Styling/SlateWidgetStyleAsset.h"
+#include "Widgets/SLeafWidget.h"
+
+class UWidget;
+
+/**
+ * A simple slate that renders the meter in single material and modifies the material on value change.
+ *
+ */
+class AUDIOWIDGETS_API SAudioMaterialMeter : public SLeafWidget
+{
+public:
+	SLATE_BEGIN_ARGS(SAudioMaterialMeter)
+	{}
+
+	/** The owner object*/
+	SLATE_ARGUMENT(TWeakObjectPtr<UObject>, Owner)
+
+	/** A value that drives how the Meter is rendered*/
+	SLATE_ATTRIBUTE(float, ValueAttribute)
+
+	/** The style used to draw the meter. */
+	SLATE_STYLE_ARGUMENT(FAudioMaterialMeterStyle, AudioMaterialMeterStyle)
+
+	/** Called when the value is changed in the Meter */
+	SLATE_EVENT(FOnFloatValueChanged, OnValueChanged)
+
+	SLATE_END_ARGS()
+
+	/**
+	* Construct the widget.
+	*/
+	void Construct(const FArguments& InArgs);
+
+	// SWidget overrides
+	virtual int32 OnPaint(const FPaintArgs& Args, const FGeometry& AllottedGeometry, const FSlateRect& MyCullingRect, FSlateWindowElementList& OutDrawElements, int32 LayerId, const FWidgetStyle& InWidgetStyle, bool bParentEnabled) const override;
+	virtual FVector2D ComputeDesiredSize(float) const override;
+
+	/** Set the Value attribute */
+	void SetValue(float InValueAttribute);
+
+	/** Apply new material to be used to render the Slate.*/
+	void ApplyNewMaterial();
+
+public:
+
+	// Holds a delegate that is executed when the slider's value changed.
+	FOnFloatValueChanged OnValueChanged;	
+
+private:
+
+	/**
+	* Commits the specified meter value.
+	*
+	* @param NewValue The value to commit.
+	*/
+	void CommitValue(float NewValue);
+
+private:
+
+	// Holds the owner of the Slate
+	TWeakObjectPtr<UObject> Owner;
+
+	// Holds the style for the Slate
+	const FAudioMaterialMeterStyle* AudioMaterialMeterStyle = nullptr;
+
+	//Holds the current value
+	TAttribute<float> ValueAttribute = 0.f;
+
+	//Width & height
+	//TODO: Currently hardcoded for testing, make more dynamic.
+	float SlateWidth = 32.f;
+	float SlateHeight = 512.f;
+
+};
