Added some convenience helpers to SAudioSpectrogram. Renamed EAudioColorMap enum to EAudioColorGradient.

#changelist validated
#virtualized

diff --git a/Engine/Plugins/Runtime/AudioWidgets/Source/AudioWidgets/Private/AudioColorMapper.cpp b/Engine/Plugins/Runtime/AudioWidgets/Source/AudioWidgets/Private/AudioColorMapper.cpp
--- a/Engine/Plugins/Runtime/AudioWidgets/Source/AudioWidgets/Private/AudioColorMapper.cpp
+++ b/Engine/Plugins/Runtime/AudioWidgets/Source/AudioWidgets/Private/AudioColorMapper.cpp
@@ -9,9 +9,9 @@ FColor FAudioColorMapper::GetColorFromValue(const float Value) const
 	switch (ColorMap)
 	{
 	default:
-	case EAudioColorMap::BlackToWhite:
+	case EAudioColorGradient::BlackToWhite:
 		return ColorMapBlackToWhite(ScaledClampedValue);
-	case EAudioColorMap::WhiteToBlack:
+	case EAudioColorGradient::WhiteToBlack:
 		return ColorMapWhiteToBlack(ScaledClampedValue);
 	}
 }
diff --git a/Engine/Plugins/Runtime/AudioWidgets/Source/AudioWidgets/Private/SAudioSpectrogram.cpp b/Engine/Plugins/Runtime/AudioWidgets/Source/AudioWidgets/Private/SAudioSpectrogram.cpp
--- a/Engine/Plugins/Runtime/AudioWidgets/Source/AudioWidgets/Private/SAudioSpectrogram.cpp
+++ b/Engine/Plugins/Runtime/AudioWidgets/Source/AudioWidgets/Private/SAudioSpectrogram.cpp
@@ -2,9 +2,11 @@
 
 #include "SAudioSpectrogram.h"
 
+#include "ConstantQ.h"
 #include "Framework/Application/SlateApplication.h"
 #include "Framework/MultiBox/MultiBoxBuilder.h"
 #include "Layout/WidgetPath.h"
+#include "SynesthesiaSpectrumAnalysis.h"
 
 #define LOCTEXT_NAMESPACE "SAudioSpectrogram"
 
@@ -31,6 +33,42 @@ void SAudioSpectrogram::AddFrame(const FAudioSpectrogramFrameData& SpectrogramFr
 	SpectrogramViewport->AddFrame(SpectrogramFrameData);
 }
 
+void SAudioSpectrogram::AddFrame(const FSynesthesiaSpectrumResults& SpectrumResults, const float SampleRate)
+{
+	const TConstArrayView<float> SquaredMagnitudes(SpectrumResults.SpectrumValues);
+	const int32 FFTSize = 2 * (SquaredMagnitudes.Num() - 1);
+	const float BinSize = SampleRate / FFTSize;
+
+	// Just ignoring the Nyquist sample here to get us a power of two length:
+	ensure(FMath::IsPowerOfTwo(SquaredMagnitudes.Num() - 1));
+	const TConstArrayView<float> SquaredMagnitudesNoNyquist = SquaredMagnitudes.LeftChop(1);
+
+	const FAudioSpectrogramFrameData SpectrogramFrameData
+	{
+		.SquaredMagnitudes = SquaredMagnitudesNoNyquist,
+		.MinFrequency = 0.0f,
+		.MaxFrequency = (SquaredMagnitudesNoNyquist.Num() - 1) * BinSize,
+		.bLogSpacedFreqencies = false,
+	};
+
+	AddFrame(SpectrogramFrameData);
+}
+
+void SAudioSpectrogram::AddFrame(const FConstantQResults& ConstantQResults, const float StartingFrequencyHz, const float NumBandsPerOctave)
+{
+	const int32 NumBands = ConstantQResults.SpectrumValues.Num();
+
+	const FAudioSpectrogramFrameData SpectrogramFrameData
+	{
+		.SquaredMagnitudes = ConstantQResults.SpectrumValues,
+		.MinFrequency = StartingFrequencyHz,
+		.MaxFrequency = StartingFrequencyHz * FMath::Pow(2.0f, (NumBands - 1) / NumBandsPerOctave),
+		.bLogSpacedFreqencies = true,
+	};
+
+	AddFrame(SpectrogramFrameData);
+}
+
 TSharedRef<const FExtensionBase> SAudioSpectrogram::AddContextMenuExtension(EExtensionHook::Position HookPosition, const TSharedPtr<FUICommandList>& CommandList, const FMenuExtensionDelegate& MenuExtensionDelegate)
 {
 	if (!ContextMenuExtender.IsValid())
@@ -181,11 +219,11 @@ TSharedRef<SWidget> SAudioSpectrogram::BuildDefaultContextMenu()
 
 void SAudioSpectrogram::BuildColorMapSubMenu(FMenuBuilder& SubMenu)
 {
-	const UEnum* EnumClass = StaticEnum<EAudioColorMap>();
+	const UEnum* EnumClass = StaticEnum<EAudioColorGradient>();
 	const int32 NumEnumValues = EnumClass->NumEnums() - 1; // Exclude 'MAX' enum value.
 	for (int32 Index = 0; Index < NumEnumValues; Index++)
 	{
-		const auto EnumValue = static_cast<EAudioColorMap>(EnumClass->GetValueByIndex(Index));
+		const auto EnumValue = static_cast<EAudioColorGradient>(EnumClass->GetValueByIndex(Index));
 
 		SubMenu.AddMenuEntry(
 			EnumClass->GetDisplayNameTextByIndex(Index),
diff --git a/Engine/Plugins/Runtime/AudioWidgets/Source/AudioWidgets/Public/AudioColorMapper.h b/Engine/Plugins/Runtime/AudioWidgets/Source/AudioWidgets/Public/AudioColorMapper.h
--- a/Engine/Plugins/Runtime/AudioWidgets/Source/AudioWidgets/Public/AudioColorMapper.h
+++ b/Engine/Plugins/Runtime/AudioWidgets/Source/AudioWidgets/Public/AudioColorMapper.h
@@ -6,7 +6,7 @@
 #include "UObject/ObjectMacros.h"
 
 UENUM(BlueprintType)
-enum class EAudioColorMap : uint8
+enum class EAudioColorGradient : uint8
 {
 	BlackToWhite,
 	WhiteToBlack,
@@ -18,7 +18,7 @@ enum class EAudioColorMap : uint8
 class FAudioColorMapper
 {
 public:
-	FAudioColorMapper(const float InMinValue, const float InMaxValue, const EAudioColorMap InColorMap)
+	FAudioColorMapper(const float InMinValue, const float InMaxValue, const EAudioColorGradient InColorMap)
 		: ValueScaling(1.0f / (InMaxValue - InMinValue))
 		, MinValue(InMinValue)
 		, ColorMap(InColorMap)
@@ -34,5 +34,5 @@ public:
 private:
 	const float ValueScaling;
 	const float MinValue;
-	const EAudioColorMap ColorMap;
+	const EAudioColorGradient ColorMap;
 };
diff --git a/Engine/Plugins/Runtime/AudioWidgets/Source/AudioWidgets/Public/AudioSpectrogramViewport.h b/Engine/Plugins/Runtime/AudioWidgets/Source/AudioWidgets/Public/AudioSpectrogramViewport.h
--- a/Engine/Plugins/Runtime/AudioWidgets/Source/AudioWidgets/Public/AudioSpectrogramViewport.h
+++ b/Engine/Plugins/Runtime/AudioWidgets/Source/AudioWidgets/Public/AudioSpectrogramViewport.h
@@ -34,7 +34,7 @@ struct FAudioSpectrogramViewportRenderParams
 	float ViewMaxFrequency = 0.0f;
 	float ColorMapMinSoundLevel = 0.0f;
 	float ColorMapMaxSoundLevel = 0.0f;
-	EAudioColorMap ColorMap = EAudioColorMap::BlackToWhite;
+	EAudioColorGradient ColorMap = EAudioColorGradient::BlackToWhite;
 	EAudioSpectrogramFrequencyAxisScale FrequencyAxisScale = EAudioSpectrogramFrequencyAxisScale::Logarithmic;
 	EAudioSpectrogramFrequencyAxisPixelBucketMode FrequencyAxisPixelBucketMode = EAudioSpectrogramFrequencyAxisPixelBucketMode::Average;
 };
diff --git a/Engine/Plugins/Runtime/AudioWidgets/Source/AudioWidgets/Public/SAudioSpectrogram.h b/Engine/Plugins/Runtime/AudioWidgets/Source/AudioWidgets/Public/SAudioSpectrogram.h
--- a/Engine/Plugins/Runtime/AudioWidgets/Source/AudioWidgets/Public/SAudioSpectrogram.h
+++ b/Engine/Plugins/Runtime/AudioWidgets/Source/AudioWidgets/Public/SAudioSpectrogram.h
@@ -8,6 +8,9 @@
 #include "Framework/SlateDelegates.h"
 #include "Widgets/SCompoundWidget.h"
 
+struct FSynesthesiaSpectrumResults;
+struct FConstantQResults;
+
 /**
  * Slate Widget for rendering a time-frequency representation of a series of audio power spectra.
  */
@@ -19,7 +22,7 @@ public:
 		, _ViewMaxFrequency(20000.0f)
 		, _ColorMapMinSoundLevel(-84.0f)
 		, _ColorMapMaxSoundLevel(12.0f)
-		, _ColorMap(EAudioColorMap::BlackToWhite)
+		, _ColorMap(EAudioColorGradient::BlackToWhite)
 		, _FrequencyAxisScale(EAudioSpectrogramFrequencyAxisScale::Logarithmic)
 		, _FrequencyAxisPixelBucketMode(EAudioSpectrogramFrequencyAxisPixelBucketMode::Average)
 		, _Orientation(EOrientation::Orient_Horizontal)
@@ -29,7 +32,7 @@ public:
 		SLATE_ATTRIBUTE(float, ViewMaxFrequency)
 		SLATE_ATTRIBUTE(float, ColorMapMinSoundLevel)
 		SLATE_ATTRIBUTE(float, ColorMapMaxSoundLevel)
-		SLATE_ATTRIBUTE(EAudioColorMap, ColorMap)
+		SLATE_ATTRIBUTE(EAudioColorGradient, ColorMap)
 		SLATE_ATTRIBUTE(EAudioSpectrogramFrequencyAxisScale, FrequencyAxisScale)
 		SLATE_ATTRIBUTE(EAudioSpectrogramFrequencyAxisPixelBucketMode, FrequencyAxisPixelBucketMode)
 		SLATE_ATTRIBUTE(EOrientation, Orientation)
@@ -43,11 +46,17 @@ public:
 	/** Add the data for one spectrum frame to the spectrogram display */
 	void AddFrame(const FAudioSpectrogramFrameData& SpectrogramFrameData);
 
+	/** Add the data for one spectrum frame to the spectrogram display (convenience helper for when using USynesthesiaSpectrumAnalyzer). Results are required to be EAudioSpectrumType::PowerSpectrum */
+	void AddFrame(const FSynesthesiaSpectrumResults& SpectrumResults, const float SampleRate);
+
+	/** Add the data for one spectrum frame to the spectrogram display (convenience helper for when using UConstantQAnalyzer). Results are required to be EAudioSpectrumType::PowerSpectrum */
+	void AddFrame(const FConstantQResults& ConstantQResults, const float StartingFrequencyHz, const float NumBandsPerOctave);
+
 	void SetViewMinFrequency(const float InViewMinFrequency) { ViewMinFrequency = InViewMinFrequency; }
 	void SetViewMaxFrequency(const float InViewMaxFrequency) { ViewMaxFrequency = InViewMaxFrequency; }
 	void SetColorMapMinSoundLevel(const float InColorMapMinSoundLevel) { ColorMapMinSoundLevel = InColorMapMinSoundLevel; }
 	void SetColorMapMaxSoundLevel(const float InColorMapMaxSoundLevel) { ColorMapMaxSoundLevel = InColorMapMaxSoundLevel; }
-	void SetColorMap(const EAudioColorMap InColorMap) { ColorMap = InColorMap; }
+	void SetColorMap(const EAudioColorGradient InColorMap) { ColorMap = InColorMap; }
 	void SetFrequencyAxisScale(const EAudioSpectrogramFrequencyAxisScale InFrequencyAxisScale) { FrequencyAxisScale = InFrequencyAxisScale; }
 	void SetFrequencyAxisPixelBucketMode(const EAudioSpectrogramFrequencyAxisPixelBucketMode InFrequencyAxisPixelBucketMode) { FrequencyAxisPixelBucketMode = InFrequencyAxisPixelBucketMode; }
 	void SetOrientation(const EOrientation InOrientation) { Orientation = InOrientation; }
@@ -79,7 +88,7 @@ private:
 	TAttribute<float> ViewMaxFrequency;
 	TAttribute<float> ColorMapMinSoundLevel;
 	TAttribute<float> ColorMapMaxSoundLevel;
-	TAttribute<EAudioColorMap> ColorMap;
+	TAttribute<EAudioColorGradient> ColorMap;
 	TAttribute<EAudioSpectrogramFrequencyAxisScale> FrequencyAxisScale;
 	TAttribute<EAudioSpectrogramFrequencyAxisPixelBucketMode> FrequencyAxisPixelBucketMode;
 	TAttribute<EOrientation> Orientation;
