AudioMaterialEnvelope first pass

#changelist validated
#virtualized

diff --git a/Engine/Plugins/Runtime/AudioWidgets/Source/AudioWidgets/Private/AudioMaterialSlate/AudioMaterialEnvelope.cpp b/Engine/Plugins/Runtime/AudioWidgets/Source/AudioWidgets/Private/AudioMaterialSlate/AudioMaterialEnvelope.cpp
new file mode 100644
--- /dev/null
+++ b/Engine/Plugins/Runtime/AudioWidgets/Source/AudioWidgets/Private/AudioMaterialSlate/AudioMaterialEnvelope.cpp
@@ -0,0 +1,46 @@
+// Copyright Epic Games, Inc. All Rights Reserved.
+
+
+#include "AudioMaterialSlate/AudioMaterialEnvelope.h"
+#include "AudioMaterialSlate/AudioMaterialSlateStyles.h"
+#include "AudioMaterialSlate/SAudioMaterialEnvelope.h"
+#include "Widgets/SWeakWidget.h"
+#include "Components/AudioComponent.h"
+
+#define LOCTEXT_NAMESPACE "AudioWidgets"
+
+#if WITH_EDITOR
+const FText UAudioMaterialEnvelope::GetPaletteCategory()
+{
+	return LOCTEXT("PaletteCategory", "AudioMaterial");
+}
+#endif
+
+void UAudioMaterialEnvelope::SynchronizeProperties()
+{
+	Super::SynchronizeProperties();
+
+	if (!EnvelopeCurve.IsValid())
+	{
+		return;
+	}
+
+	EnvelopeCurve->ApplyNewMaterial();
+}
+
+void UAudioMaterialEnvelope::ReleaseSlateResources(bool bReleaseChildren)
+{
+	EnvelopeCurve.Reset();
+}
+
+TSharedRef<SWidget> UAudioMaterialEnvelope::RebuildWidget()
+{
+	EnvelopeCurve = SNew(SAudioMaterialEnvelope)
+		.Owner(this)
+		.AudioMaterialMeterStyle(&WidgetStyle)
+		.EnvelopeSettings(&EnvelopeSettings);
+
+	return EnvelopeCurve.ToSharedRef();
+}
+
+#undef LOCTEXT_NAMESPACE
diff --git a/Engine/Plugins/Runtime/AudioWidgets/Source/AudioWidgets/Private/AudioMaterialSlate/AudioMaterialKnob.cpp b/Engine/Plugins/Runtime/AudioWidgets/Source/AudioWidgets/Private/AudioMaterialSlate/AudioMaterialKnob.cpp
--- a/Engine/Plugins/Runtime/AudioWidgets/Source/AudioWidgets/Private/AudioMaterialSlate/AudioMaterialKnob.cpp
+++ b/Engine/Plugins/Runtime/AudioWidgets/Source/AudioWidgets/Private/AudioMaterialSlate/AudioMaterialKnob.cpp
@@ -43,11 +43,6 @@ void UAudioMaterialKnob::SetValue(float InValue)
 	if (Knob.IsValid())
 	{
 		Knob->SetValue(InValue);
-	}
-
-	if (Value != InValue)
-	{
-		Value = InValue;
 		HandleOnKnobValueChanged(InValue);
 	}
 }
@@ -65,8 +60,11 @@ TSharedRef<SWidget> UAudioMaterialKnob::RebuildWidget()
 
 void UAudioMaterialKnob::HandleOnKnobValueChanged(float InValue)
 {
-	Value = InValue;
-	OnKnobValueChanged.Broadcast(InValue);	
+	if (Value != InValue)
+	{
+		Value = InValue;
+		OnKnobValueChanged.Broadcast(InValue);	
+	}
 }
 
 #undef LOCTEXT_NAMESPACE
diff --git a/Engine/Plugins/Runtime/AudioWidgets/Source/AudioWidgets/Private/AudioMaterialSlate/AudioMaterialSlateStyles.cpp b/Engine/Plugins/Runtime/AudioWidgets/Source/AudioWidgets/Private/AudioMaterialSlate/AudioMaterialSlateStyles.cpp
--- a/Engine/Plugins/Runtime/AudioWidgets/Source/AudioWidgets/Private/AudioMaterialSlate/AudioMaterialSlateStyles.cpp
+++ b/Engine/Plugins/Runtime/AudioWidgets/Source/AudioWidgets/Private/AudioMaterialSlate/AudioMaterialSlateStyles.cpp
@@ -119,3 +119,28 @@ void FAudioMaterialMeterStyle::GetResources(TArray<const FSlateBrush*>& OutBrush
 {
 	// Add any brush resources here so that Slate can correctly atlas and reference them
 }
+
+
+FAudioMaterialEnvelopeStyle::FAudioMaterialEnvelopeStyle()
+:CurveColor(FLinearColor(0.15f, 0.45f, 0.75f, 1.f)),
+BackgroundColor(FLinearColor(0.01f, 0.01f, 0.01f, 1.f)),
+OutlineColor(FLinearColor(0.05f, 0.05f, 0.05f, 1.f))
+{
+}
+
+FAudioMaterialEnvelopeStyle::~FAudioMaterialEnvelopeStyle()
+{
+}
+
+const FName FAudioMaterialEnvelopeStyle::TypeName(TEXT("FAudioMaterialKnobStyle"));
+
+const FAudioMaterialEnvelopeStyle& FAudioMaterialEnvelopeStyle::GetDefault()
+{
+	static FAudioMaterialEnvelopeStyle Default;
+	return Default;
+}
+
+void FAudioMaterialEnvelopeStyle::GetResources(TArray<const FSlateBrush*>& OutBrushes) const
+{
+	// Add any brush resources here so that Slate can correctly atlas and reference them
+}
diff --git a/Engine/Plugins/Runtime/AudioWidgets/Source/AudioWidgets/Private/AudioMaterialSlate/SAudioMaterialEnvelope.cpp b/Engine/Plugins/Runtime/AudioWidgets/Source/AudioWidgets/Private/AudioMaterialSlate/SAudioMaterialEnvelope.cpp
new file mode 100644
--- /dev/null
+++ b/Engine/Plugins/Runtime/AudioWidgets/Source/AudioWidgets/Private/AudioMaterialSlate/SAudioMaterialEnvelope.cpp
@@ -0,0 +1,77 @@
+// Copyright Epic Games, Inc. All Rights Reserved.
+
+
+#include "AudioMaterialSlate/SAudioMaterialEnvelope.h"
+#include "AudioMaterialSlate/AudioMaterialEnvelope.h"
+#include "Components/AudioComponent.h"
+#include "SlateOptMacros.h"
+#include "Styling/SlateBrush.h"
+
+
+BEGIN_SLATE_FUNCTION_BUILD_OPTIMIZATION
+
+void SAudioMaterialEnvelope::Construct(const FArguments& InArgs)
+{
+	Owner = InArgs._Owner;
+	EnvelopeSettings = InArgs._EnvelopeSettings;
+	AudioMaterialMeterStyle = InArgs._AudioMaterialMeterStyle;
+
+	ApplyNewMaterial();
+}
+
+END_SLATE_FUNCTION_BUILD_OPTIMIZATION
+
+int32 SAudioMaterialEnvelope::OnPaint(const FPaintArgs& Args, const FGeometry& AllottedGeometry, const FSlateRect& MyCullingRect, FSlateWindowElementList& OutDrawElements, int32 LayerId, const FWidgetStyle& InWidgetStyle, bool bParentEnabled) const
+{
+	if (AudioMaterialMeterStyle && EnvelopeSettings)
+	{
+		UMaterialInstanceDynamic* DynamicMaterial = AudioMaterialMeterStyle->GetDynamicMaterial();
+		if (IsValid(DynamicMaterial))
+		{
+
+			DynamicMaterial->SetScalarParameterValue(FName("A_Curve"), EnvelopeSettings->AttackCurve);
+			DynamicMaterial->SetScalarParameterValue(FName("A_Int"), EnvelopeSettings->AttackValue);
+			DynamicMaterial->SetScalarParameterValue(FName("A_Time"), EnvelopeSettings->AttackTime);
+
+			DynamicMaterial->SetScalarParameterValue(FName("D_Curve"), EnvelopeSettings->DecayCurve);
+			DynamicMaterial->SetScalarParameterValue(FName("D_Time"), EnvelopeSettings->DecayTime);
+
+			if (EnvelopeSettings->EnvelopeType == EAudioMaterialEnvelopeType::ADSR)
+			{
+				DynamicMaterial->SetScalarParameterValue(FName("R_Curve"), EnvelopeSettings->ReleaseCurve);
+				DynamicMaterial->SetScalarParameterValue(FName("R_Time"), EnvelopeSettings->ReleaseTime);
+
+				DynamicMaterial->SetScalarParameterValue(FName("S_Int"), EnvelopeSettings->SustainValue);
+			}
+
+			DynamicMaterial->SetVectorParameterValue(FName("MainColor"), AudioMaterialMeterStyle->CurveColor);
+			DynamicMaterial->SetVectorParameterValue(FName("BoxBG"), AudioMaterialMeterStyle->BackgroundColor);
+			DynamicMaterial->SetVectorParameterValue(FName("BoxOutline"), AudioMaterialMeterStyle->OutlineColor);
+		}
+
+		const bool bEnabled = ShouldBeEnabled(bParentEnabled);
+		const ESlateDrawEffect DrawEffects = bEnabled ? ESlateDrawEffect::None : ESlateDrawEffect::DisabledEffect;
+
+		const FLinearColor FinalColorAndOpacity(InWidgetStyle.GetColorAndOpacityTint());
+
+		FSlateBrush Brush;
+		Brush.SetResourceObject(DynamicMaterial);
+		FSlateDrawElement::MakeBox(OutDrawElements, LayerId++, AllottedGeometry.ToPaintGeometry(), &Brush, DrawEffects, FinalColorAndOpacity);
+	}
+
+	return LayerId;
+}
+
+FVector2D SAudioMaterialEnvelope::ComputeDesiredSize(float) const
+{
+	FVector2D Vector = FVector2D(SlateWidth, SlateHeight);
+	return Vector;
+}
+
+void SAudioMaterialEnvelope::ApplyNewMaterial()
+{
+	if (AudioMaterialMeterStyle)
+	{
+		AudioMaterialMeterStyle->CreateDynamicMaterial(Owner.Get());
+	}
+}
diff --git a/Engine/Plugins/Runtime/AudioWidgets/Source/AudioWidgets/Public/AudioMaterialSlate/AudioMaterialEnvelope.h b/Engine/Plugins/Runtime/AudioWidgets/Source/AudioWidgets/Public/AudioMaterialSlate/AudioMaterialEnvelope.h
new file mode 100644
--- /dev/null
+++ b/Engine/Plugins/Runtime/AudioWidgets/Source/AudioWidgets/Public/AudioMaterialSlate/AudioMaterialEnvelope.h
@@ -0,0 +1,59 @@
+// Copyright Epic Games, Inc. All Rights Reserved.
+
+#pragma once
+
+#include "AudioMaterialEnvelopeSettings.h"
+#include "AudioMaterialSlateStyles.h"
+#include "Components/Widget.h"
+#include "Delegates/Delegate.h"
+#include "AudioMaterialEnvelope.generated.h"
+
+class SAudioMaterialEnvelope;
+
+/**
+ * A simple widget that shows a envelope curve Depending on given AudioMaterialEnvelopeSetings
+ * Rendered by using material instead of texture.
+ *
+ * * No Children
+ */
+UCLASS()
+class AUDIOWIDGETS_API UAudioMaterialEnvelope : public UWidget
+{
+	GENERATED_BODY()
+
+public:
+
+	/** The Envelope's style */
+	UPROPERTY(EditAnywhere, BlueprintReadWrite, Category = "Style", meta = (DisplayName = "Style"))
+	FAudioMaterialEnvelopeStyle WidgetStyle;
+
+	/**Envelope settings*/
+	UPROPERTY(EditAnywhere, BlueprintReadWrite, Category = "Appearance")
+	FAudioMaterialEnvelopeSettings EnvelopeSettings;
+
+public:
+
+#if WITH_EDITOR
+	virtual const FText GetPaletteCategory() override;
+#endif
+
+	// UWidget
+	virtual void SynchronizeProperties() override;
+	// End of UWidget
+
+	// UVisual
+	virtual void ReleaseSlateResources(bool bReleaseChildren) override;
+	// End of UVisual
+
+protected:
+
+	// UWidget
+	virtual TSharedRef<SWidget> RebuildWidget() override;
+	// End of UWidget
+
+private:
+
+	/** Native Slate Widget */
+	TSharedPtr<SAudioMaterialEnvelope> EnvelopeCurve;
+
+};
diff --git a/Engine/Plugins/Runtime/AudioWidgets/Source/AudioWidgets/Public/AudioMaterialSlate/AudioMaterialEnvelopeSettings.h b/Engine/Plugins/Runtime/AudioWidgets/Source/AudioWidgets/Public/AudioMaterialSlate/AudioMaterialEnvelopeSettings.h
new file mode 100644
--- /dev/null
+++ b/Engine/Plugins/Runtime/AudioWidgets/Source/AudioWidgets/Public/AudioMaterialSlate/AudioMaterialEnvelopeSettings.h
@@ -0,0 +1,90 @@
+// Copyright Epic Games, Inc. All Rights Reserved.
+
+#pragma once
+
+#include "AudioMaterialEnvelopeSettings.generated.h"
+
+/**
+*Envelope curve types
+*A: Attack
+*D: Decay
+*S: Sustain
+*R: Release
+*/
+UENUM(BlueprintType)
+enum class EAudioMaterialEnvelopeType : uint8
+{
+	AD,
+	ADSR
+};
+
+USTRUCT(BlueprintType)
+struct AUDIOWIDGETS_API FAudioMaterialEnvelopeSettings
+{
+	GENERATED_BODY();
+
+public:	
+	
+	/**
+	* The Type of the envelope curve.
+	*/
+	UPROPERTY(EditAnywhere, Category = "Envelope|Type")
+	EAudioMaterialEnvelopeType EnvelopeType = EAudioMaterialEnvelopeType::ADSR;
+
+	/**
+	* Curve of the envelopes attack stage.
+	* Attack is the time taken for the rise of the level from zero to a given value.
+	*/
+	UPROPERTY(EditAnywhere, BlueprintReadWrite, Category = "Envelope|Attack")
+	float AttackCurve = 1.0;
+
+	/**
+	* Value of the envelopes attack stage.
+	* Attack is the time taken for the rise of the level from zero to a given value.
+	*/
+	UPROPERTY(EditAnywhere, BlueprintReadWrite, Category = "Envelope|Attack")
+	float AttackValue = 1.0;
+
+	/**
+	* Time the Value reaches the Attack stage.
+	* Attack is the time taken for the rise of the level from zero to a given value.
+	*/
+	UPROPERTY(EditAnywhere, BlueprintReadWrite, Category = "Envelope|Attack")
+	float AttackTime = 1.0;
+	
+	/**
+	* Curve of the envelopes Decay stage.
+	* Decay is the time taken for the level to reduce from the attack level to the sustain level.
+	*/
+	UPROPERTY(EditAnywhere, BlueprintReadWrite, Category = "Envelope|Decay")
+	float DecayCurve = 1.0;	
+	
+	/**
+	* Time that takes to reach the level of the Sustain stage.
+	* Decay is the time taken for the level to reduce from the attack level to the sustain level.
+	*/	
+	UPROPERTY(EditAnywhere, BlueprintReadWrite, Category = "Envelope|Decay")
+	float DecayTime = 1.0;	
+	
+	/**
+	* Value of the envelopes Sustain stage.
+	* Sustain is the level maintained until release stage.
+	*/
+	UPROPERTY(EditAnywhere, BlueprintReadWrite, meta = (EditCondition = "EnvelopeType==EAudioMaterialEnvelopeType::ADSR", EditConditionHides), Category = "Envelope|Sustain")
+	float SustainValue = 1.0;
+
+	/**
+	* Curve of the envelopes Release stage.
+	* Release is the time taken for the level to decay from sustain to zero.
+	*/
+	UPROPERTY(EditAnywhere, BlueprintReadWrite, meta = (EditCondition = "EnvelopeType==EAudioMaterialEnvelopeType::ADSR", EditConditionHides), Category = "Envelope|Release")
+	float ReleaseCurve = 1.0;	
+	
+	/**
+	* Time that takes to reach zero level
+	* Release is the time taken for the level to decay from sustain to zero.
+	*/
+	UPROPERTY(EditAnywhere, BlueprintReadWrite, meta = (EditCondition = "EnvelopeType==EAudioMaterialEnvelopeType::ADSR", EditConditionHides), Category = "Envelope|Release")
+	float ReleaseTime = 1.0;
+	
+};
\ No newline at end of file
diff --git a/Engine/Plugins/Runtime/AudioWidgets/Source/AudioWidgets/Public/AudioMaterialSlate/AudioMaterialSlateStyles.h b/Engine/Plugins/Runtime/AudioWidgets/Source/AudioWidgets/Public/AudioMaterialSlate/AudioMaterialSlateStyles.h
--- a/Engine/Plugins/Runtime/AudioWidgets/Source/AudioWidgets/Public/AudioMaterialSlate/AudioMaterialSlateStyles.h
+++ b/Engine/Plugins/Runtime/AudioWidgets/Source/AudioWidgets/Public/AudioMaterialSlate/AudioMaterialSlateStyles.h
@@ -182,3 +182,31 @@ struct AUDIOWIDGETS_API FAudioMaterialMeterStyle : public FAudioMaterialStyle
 	FLinearColor MeterOffFillColor;
 	
 };
+
+/**
+ *Represents the appearance of an Audio Material Envelope
+ */
+USTRUCT(BlueprintType)
+struct AUDIOWIDGETS_API FAudioMaterialEnvelopeStyle : public FAudioMaterialStyle
+{
+	GENERATED_USTRUCT_BODY()
+
+	FAudioMaterialEnvelopeStyle();
+	virtual ~FAudioMaterialEnvelopeStyle();
+
+	// FSlateWidgetStyle
+	virtual void GetResources(TArray<const FSlateBrush*>& OutBrushes) const override;
+	static const FName TypeName;
+	virtual const FName GetTypeName() const override { return TypeName; };
+	static const FAudioMaterialEnvelopeStyle& GetDefault();
+
+	UPROPERTY(EditAnywhere, BlueprintReadWrite, Category = Appearance)
+	FLinearColor CurveColor;	
+	
+	UPROPERTY(EditAnywhere, BlueprintReadWrite, Category = Appearance)
+	FLinearColor BackgroundColor;	
+	
+	UPROPERTY(EditAnywhere, BlueprintReadWrite, Category = Appearance)
+	FLinearColor OutlineColor;
+	
+};
diff --git a/Engine/Plugins/Runtime/AudioWidgets/Source/AudioWidgets/Public/AudioMaterialSlate/SAudioMaterialEnvelope.h b/Engine/Plugins/Runtime/AudioWidgets/Source/AudioWidgets/Public/AudioMaterialSlate/SAudioMaterialEnvelope.h
new file mode 100644
--- /dev/null
+++ b/Engine/Plugins/Runtime/AudioWidgets/Source/AudioWidgets/Public/AudioMaterialSlate/SAudioMaterialEnvelope.h
@@ -0,0 +1,61 @@
+// Copyright Epic Games, Inc. All Rights Reserved.
+
+#pragma once
+
+#include "AudioMaterialSlate/AudioMaterialSlateStyles.h"
+#include "AudioMaterialEnvelopeSettings.h"
+#include "Framework/SlateDelegates.h"
+#include "Styling/ISlateStyle.h"
+#include "Styling/SlateWidgetStyleAsset.h"
+#include "Widgets/SLeafWidget.h"
+
+/**
+ * A simple slate that renders envelope curve in a single material and modifies the material on value change.
+ *
+ */
+class AUDIOWIDGETS_API SAudioMaterialEnvelope : public SLeafWidget
+{
+public:
+	SLATE_BEGIN_ARGS(SAudioMaterialEnvelope)
+	{}
+
+	/** The owner object*/
+	SLATE_ARGUMENT(TWeakObjectPtr<UObject>, Owner)
+
+	/** Drives how the Envelope Curve is rendered*/
+	SLATE_ARGUMENT(FAudioMaterialEnvelopeSettings*, EnvelopeSettings)
+
+	/** The style used to draw the Envelope. */
+	SLATE_STYLE_ARGUMENT(FAudioMaterialEnvelopeStyle, AudioMaterialMeterStyle)
+
+	SLATE_END_ARGS()
+
+	/**
+	* Construct the widget.
+	*/
+	void Construct(const FArguments& InArgs);
+
+	// SWidget overrides
+	virtual int32 OnPaint(const FPaintArgs& Args, const FGeometry& AllottedGeometry, const FSlateRect& MyCullingRect, FSlateWindowElementList& OutDrawElements, int32 LayerId, const FWidgetStyle& InWidgetStyle, bool bParentEnabled) const override;
+	virtual FVector2D ComputeDesiredSize(float) const override;
+
+	/** Apply new material to be used to render the Slate.*/
+	void ApplyNewMaterial();
+
+private:
+
+	// Holds the owner of the Slate
+	TWeakObjectPtr<UObject> Owner;
+
+	// Holds the style for the Slate
+	const FAudioMaterialEnvelopeStyle* AudioMaterialMeterStyle = nullptr;
+
+	//Holds the current Envelope settings
+	const FAudioMaterialEnvelopeSettings* EnvelopeSettings = nullptr;
+
+	//Width & height
+	//TODO: Currently hardcoded for testing, make more dynamic.
+	float SlateWidth = 256.f;
+	float SlateHeight = 256.f;
+
+};
