Support other EAudioSpectrumTypes in SAudioSpectrogram.

#changelist validated
#virtualized

diff --git a/Engine/Plugins/Runtime/AudioWidgets/Source/AudioWidgets/Private/AudioSpectrogramViewport.cpp b/Engine/Plugins/Runtime/AudioWidgets/Source/AudioWidgets/Private/AudioSpectrogramViewport.cpp
--- a/Engine/Plugins/Runtime/AudioWidgets/Source/AudioWidgets/Private/AudioSpectrogramViewport.cpp
+++ b/Engine/Plugins/Runtime/AudioWidgets/Source/AudioWidgets/Private/AudioSpectrogramViewport.cpp
@@ -251,7 +251,7 @@ void FAudioSpectrogramViewport::UpdateTextureData(const bool bInvalidateCachedPi
 ////////////////////////////////////////////////////////////////////////////////////////////////////////////////////////////////////////////////////////
 
 FAudioSpectrogramViewport::FHistoryFrameData::FHistoryFrameData(const FAudioSpectrogramFrameData& SpectrogramFrameData, const EAudioSpectrogramFrequencyAxisPixelBucketMode InFrequencyAxisPixelBucketMode)
-	: NumSoundLevels(SpectrogramFrameData.SquaredMagnitudes.Num())
+	: NumSoundLevels(SpectrogramFrameData.SpectrumValues.Num())
 	, SpectrumDataMinFrequency(SpectrogramFrameData.MinFrequency)
 	, SpectrumDataMaxFrequency(SpectrogramFrameData.MaxFrequency)
 	, bSpectrumDataHasLogSpacedFreqencies(SpectrogramFrameData.bLogSpacedFreqencies)
@@ -260,12 +260,19 @@ FAudioSpectrogramViewport::FHistoryFrameData::FHistoryFrameData(const FAudioSpec
 	// Allocate required space for mip chain:
 	SetMipChainLengthUninitialized(InFrequencyAxisPixelBucketMode);
 
-	// Convert squared magnitudes to dB, and write to mip zero:
-	const float ClampMinMagnitudeSquared = FMath::Pow(10.0f, -200.0f / 10.0f); // Clamp at -200dB
-	for (int Index = 0; Index < NumSoundLevels; Index++)
+	// Copy spectrum data to mip zero, converting to dB if necessary:
+	switch (SpectrogramFrameData.SpectrumType)
 	{
-		const float MagnitudeSquared = FMath::Max(SpectrogramFrameData.SquaredMagnitudes[Index], ClampMinMagnitudeSquared);
-		SoundLevelsMipChain[Index] = 10.0f * FMath::LogX(10.0f, MagnitudeSquared);
+		case EAudioSpectrumType::MagnitudeSpectrum:
+			ArrayMagnitudeToDecibel(SpectrogramFrameData.SpectrumValues, SoundLevelsMipChain, -200.0f);
+			break;
+		case EAudioSpectrumType::PowerSpectrum:
+			ArrayPowerToDecibel(SpectrogramFrameData.SpectrumValues, SoundLevelsMipChain, -200.0f);
+			break;
+		default:
+		case EAudioSpectrumType::Decibel:
+			FMemory::Memcpy(SoundLevelsMipChain.GetData(), SpectrogramFrameData.SpectrumValues.GetData(), NumSoundLevels * sizeof(float));
+			break;
 	}
 
 	// Generate remaining mips if required:
@@ -425,3 +432,24 @@ float FAudioSpectrogramViewport::FHistoryFrameData::GetInterpolatedSoundLevel(co
 	return FMath::Lerp(MipData[IndexLo], MipData[IndexHi], LerpParam);
 }
 
+void FAudioSpectrogramViewport::FHistoryFrameData::ArrayMagnitudeToDecibel(TConstArrayView<float> InValues, TArrayView<float> OutValues, float InMinimumDb)
+{
+	const float ClampMinMagnitude = Audio::ConvertToLinear(InMinimumDb);
+	const float Scale = 20.0f / FMath::Loge(10.0f);
+	for (int32 Index = 0; Index < InValues.Num(); Index++)
+	{
+		const float Magnitude = FMath::Max(InValues[Index], ClampMinMagnitude);
+		OutValues[Index] = Scale * FMath::Loge(Magnitude);
+	}
+}
+
+void FAudioSpectrogramViewport::FHistoryFrameData::ArrayPowerToDecibel(TConstArrayView<float> InValues, TArrayView<float> OutValues, float InMinimumDb)
+{
+	const float ClampMinMagnitudeSquared = FMath::Pow(10.0f, InMinimumDb / 10.0f);
+	const float Scale = 10.0f / FMath::Loge(10.0f);
+	for (int32 Index = 0; Index < InValues.Num(); Index++)
+	{
+		const float MagnitudeSquared = FMath::Max(InValues[Index], ClampMinMagnitudeSquared);
+		OutValues[Index] = Scale * FMath::Loge(MagnitudeSquared);
+	}
+}
diff --git a/Engine/Plugins/Runtime/AudioWidgets/Source/AudioWidgets/Private/SAudioSpectrogram.cpp b/Engine/Plugins/Runtime/AudioWidgets/Source/AudioWidgets/Private/SAudioSpectrogram.cpp
--- a/Engine/Plugins/Runtime/AudioWidgets/Source/AudioWidgets/Private/SAudioSpectrogram.cpp
+++ b/Engine/Plugins/Runtime/AudioWidgets/Source/AudioWidgets/Private/SAudioSpectrogram.cpp
@@ -33,34 +33,36 @@ void SAudioSpectrogram::AddFrame(const FAudioSpectrogramFrameData& SpectrogramFr
 	SpectrogramViewport->AddFrame(SpectrogramFrameData);
 }
 
-void SAudioSpectrogram::AddFrame(const FSynesthesiaSpectrumResults& SpectrumResults, const float SampleRate)
+void SAudioSpectrogram::AddFrame(const FSynesthesiaSpectrumResults& SpectrumResults, const EAudioSpectrumType SpectrumType, const float SampleRate)
 {
-	const TConstArrayView<float> SquaredMagnitudes(SpectrumResults.SpectrumValues);
-	const int32 FFTSize = 2 * (SquaredMagnitudes.Num() - 1);
+	const TConstArrayView<float> SpectrumValues(SpectrumResults.SpectrumValues);
+	const int32 FFTSize = 2 * (SpectrumValues.Num() - 1);
 	const float BinSize = SampleRate / FFTSize;
 
 	// Just ignoring the Nyquist sample here to get us a power of two length:
-	ensure(FMath::IsPowerOfTwo(SquaredMagnitudes.Num() - 1));
-	const TConstArrayView<float> SquaredMagnitudesNoNyquist = SquaredMagnitudes.LeftChop(1);
+	ensure(FMath::IsPowerOfTwo(SpectrumValues.Num() - 1));
+	const TConstArrayView<float> SpectrumValuesNoNyquist = SpectrumValues.LeftChop(1);
 
 	const FAudioSpectrogramFrameData SpectrogramFrameData
 	{
-		.SquaredMagnitudes = SquaredMagnitudesNoNyquist,
+		.SpectrumValues = SpectrumValuesNoNyquist,
+		.SpectrumType = SpectrumType,
 		.MinFrequency = 0.0f,
-		.MaxFrequency = (SquaredMagnitudesNoNyquist.Num() - 1) * BinSize,
+		.MaxFrequency = (SpectrumValuesNoNyquist.Num() - 1) * BinSize,
 		.bLogSpacedFreqencies = false,
 	};
 
 	AddFrame(SpectrogramFrameData);
 }
 
-void SAudioSpectrogram::AddFrame(const FConstantQResults& ConstantQResults, const float StartingFrequencyHz, const float NumBandsPerOctave)
+void SAudioSpectrogram::AddFrame(const FConstantQResults& ConstantQResults, const float StartingFrequencyHz, const float NumBandsPerOctave, const EAudioSpectrumType SpectrumType)
 {
 	const int32 NumBands = ConstantQResults.SpectrumValues.Num();
 
 	const FAudioSpectrogramFrameData SpectrogramFrameData
 	{
-		.SquaredMagnitudes = ConstantQResults.SpectrumValues,
+		.SpectrumValues = ConstantQResults.SpectrumValues,
+		.SpectrumType = SpectrumType,
 		.MinFrequency = StartingFrequencyHz,
 		.MaxFrequency = StartingFrequencyHz * FMath::Pow(2.0f, (NumBands - 1) / NumBandsPerOctave),
 		.bLogSpacedFreqencies = true,
diff --git a/Engine/Plugins/Runtime/AudioWidgets/Source/AudioWidgets/Public/AudioSpectrogramViewport.h b/Engine/Plugins/Runtime/AudioWidgets/Source/AudioWidgets/Public/AudioSpectrogramViewport.h
--- a/Engine/Plugins/Runtime/AudioWidgets/Source/AudioWidgets/Public/AudioSpectrogramViewport.h
+++ b/Engine/Plugins/Runtime/AudioWidgets/Source/AudioWidgets/Public/AudioSpectrogramViewport.h
@@ -6,6 +6,7 @@
 #include "Containers/RingBuffer.h"
 #include "PixelFormat.h"
 #include "Rendering/RenderingCommon.h"
+#include "Sound/SoundSubmix.h"
 #include "Textures/SlateUpdatableTexture.h"
 
 UENUM(BlueprintType)
@@ -44,7 +45,8 @@ struct FAudioSpectrogramViewportRenderParams
  */
 struct FAudioSpectrogramFrameData
 {
-	TConstArrayView<float> SquaredMagnitudes;
+	TConstArrayView<float> SpectrumValues;
+	EAudioSpectrumType SpectrumType = EAudioSpectrumType::PowerSpectrum;
 	float MinFrequency = 0.0f;
 	float MaxFrequency = 0.0f;
 	bool bLogSpacedFreqencies = false;
@@ -112,6 +114,12 @@ private:
 		// Helper function for sampling from the given mip data at the given normalized read position. Data is linearly interpolated.
 		static float GetInterpolatedSoundLevel(const TConstArrayView<float> MipData, const float NormalizedDataPos);
 
+		// Convert magnitude values to decibel values. db = 20 * log10(val)
+		static void ArrayMagnitudeToDecibel(TConstArrayView<float> InValues, TArrayView<float> OutValues, float InMinimumDb);
+
+		// Convert power values to decibel values. db = 10 * log10(val)
+		static void ArrayPowerToDecibel(TConstArrayView<float> InValues, TArrayView<float> OutValues, float InMinimumDb);
+
 		// Immutable data describing this spectrogram frame:
 		const int32 NumSoundLevels;
 		const float SpectrumDataMinFrequency;
diff --git a/Engine/Plugins/Runtime/AudioWidgets/Source/AudioWidgets/Public/SAudioSpectrogram.h b/Engine/Plugins/Runtime/AudioWidgets/Source/AudioWidgets/Public/SAudioSpectrogram.h
--- a/Engine/Plugins/Runtime/AudioWidgets/Source/AudioWidgets/Public/SAudioSpectrogram.h
+++ b/Engine/Plugins/Runtime/AudioWidgets/Source/AudioWidgets/Public/SAudioSpectrogram.h
@@ -46,11 +46,11 @@ public:
 	/** Add the data for one spectrum frame to the spectrogram display */
 	void AddFrame(const FAudioSpectrogramFrameData& SpectrogramFrameData);
 
-	/** Add the data for one spectrum frame to the spectrogram display (convenience helper for when using USynesthesiaSpectrumAnalyzer). Results are required to be EAudioSpectrumType::PowerSpectrum */
-	void AddFrame(const FSynesthesiaSpectrumResults& SpectrumResults, const float SampleRate);
+	/** Add the data for one spectrum frame to the spectrogram display (convenience helper for when using USynesthesiaSpectrumAnalyzer) */
+	void AddFrame(const FSynesthesiaSpectrumResults& SpectrumResults, const EAudioSpectrumType SpectrumType, const float SampleRate);
 
-	/** Add the data for one spectrum frame to the spectrogram display (convenience helper for when using UConstantQAnalyzer). Results are required to be EAudioSpectrumType::PowerSpectrum */
-	void AddFrame(const FConstantQResults& ConstantQResults, const float StartingFrequencyHz, const float NumBandsPerOctave);
+	/** Add the data for one spectrum frame to the spectrogram display (convenience helper for when using UConstantQAnalyzer) */
+	void AddFrame(const FConstantQResults& ConstantQResults, const float StartingFrequencyHz, const float NumBandsPerOctave, const EAudioSpectrumType SpectrumType);
 
 	void SetViewMinFrequency(const float InViewMinFrequency) { ViewMinFrequency = InViewMinFrequency; }
 	void SetViewMaxFrequency(const float InViewMaxFrequency) { ViewMaxFrequency = InViewMaxFrequency; }
