Added default context menu for SAudioSpectrumPlot widget, following the same design as previously submitted for the SAudioSpectrogram widget.
Extended this context menu in FAudioSpectrumAnalyzer to add a control for ballistics.

diff --git a/Engine/Plugins/Runtime/AudioWidgets/Source/AudioWidgets/Private/AudioSpectrumAnalyzer.cpp b/Engine/Plugins/Runtime/AudioWidgets/Source/AudioWidgets/Private/AudioSpectrumAnalyzer.cpp
--- a/Engine/Plugins/Runtime/AudioWidgets/Source/AudioWidgets/Private/AudioSpectrumAnalyzer.cpp
+++ b/Engine/Plugins/Runtime/AudioWidgets/Source/AudioWidgets/Private/AudioSpectrumAnalyzer.cpp
@@ -3,6 +3,9 @@
 #include "AudioSpectrumAnalyzer.h"
 
 #include "DSP/EnvelopeFollower.h"
+#include "Framework/MultiBox/MultiBoxBuilder.h"
+
+#define LOCTEXT_NAMESPACE "FAudioSpectrumAnalyzer"
 
 namespace AudioWidgets
 {
@@ -11,8 +14,11 @@ namespace AudioWidgets
 			.Clipping(EWidgetClipping::ClipToBounds)
 			.DisplayFrequencyAxisLabels(false)
 			.DisplaySoundLevelAxisLabels(false)
+			.FrequencyAxisPixelBucketMode_Lambda([]() { return EAudioSpectrumPlotFrequencyAxisPixelBucketMode::Average; }) // Binding this property has the effect of hiding its context menu entry (it isn't much use for the ConstantQ analyzer).
 			.OnGetAudioSpectrumData_Raw(this, &FAudioSpectrumAnalyzer::GetAudioSpectrumData))
 	{
+		ContextMenuExtension = Widget->AddContextMenuExtension(EExtensionHook::Before, nullptr, FMenuExtensionDelegate::CreateRaw(this, &FAudioSpectrumAnalyzer::ExtendSpectrumPlotContextMenu));
+
 		Init(InNumChannels, InAudioDeviceId, InExternalAudioBus);
 	}
 
@@ -21,6 +27,11 @@ namespace AudioWidgets
 		Teardown();
 
 		Widget->UnbindOnGetAudioSpectrumData();
+
+		if (ContextMenuExtension.IsValid())
+		{
+			Widget->RemoveContextMenuExtension(ContextMenuExtension.ToSharedRef());
+		}
 	}
 
 	UAudioBus* FAudioSpectrumAnalyzer::GetAudioBus() const
@@ -71,7 +82,7 @@ namespace AudioWidgets
 				{
 					// Calculate AR smoother coefficients:
 					const float DeltaT = (SpectrumResults.TimeSeconds - PrevTimeStamp.GetValue());
-					Audio::FAttackRelease AttackRelease(1.0f / DeltaT, AttackTimeMsec, ReleaseTimeMsec, true);
+					Audio::FAttackRelease AttackRelease(1.0f / DeltaT, AttackTimeMsec, ReleaseTimeMsec, bIsAnalogAttackRelease);
 
 					// Apply AR smoothing for each frequency:
 					check(SpectrumResults.SpectrumValues.Num() == ARSmoothedSquaredMagnitudes.Num());
@@ -128,4 +139,43 @@ namespace AudioWidgets
 		check(CenterFrequencies.Num() == ARSmoothedSquaredMagnitudes.Num());
 		return FAudioPowerSpectrumData{ CenterFrequencies, ARSmoothedSquaredMagnitudes };
 	}
+
+	void FAudioSpectrumAnalyzer::ExtendSpectrumPlotContextMenu(FMenuBuilder& MenuBuilder)
+	{
+		MenuBuilder.BeginSection("AnalyzerSettings", LOCTEXT("AnalyzerSettings", "Analyzer Settings"));
+		MenuBuilder.AddSubMenu(
+			LOCTEXT("Ballistics", "Ballistics"),
+			FText(),
+			FNewMenuDelegate::CreateSP(this, &FAudioSpectrumAnalyzer::BuildBallisticsSubMenu));
+		MenuBuilder.EndSection();
+	}
+
+	void FAudioSpectrumAnalyzer::BuildBallisticsSubMenu(FMenuBuilder& SubMenu)
+	{
+		SubMenu.AddMenuEntry(
+			LOCTEXT("Analog", "Analog"),
+			FText(),
+			FSlateIcon(),
+			FUIAction(
+				FExecuteAction::CreateSPLambda(this, [this]() { bIsAnalogAttackRelease = true; }),
+				FCanExecuteAction(),
+				FIsActionChecked::CreateSPLambda(this, [this]() { return bIsAnalogAttackRelease; })
+			),
+			NAME_None,
+			EUserInterfaceActionType::ToggleButton);
+		SubMenu.AddMenuEntry(
+			LOCTEXT("Digital", "Digital"),
+			FText(),
+			FSlateIcon(),
+			FUIAction(
+				FExecuteAction::CreateSPLambda(this, [this]() { bIsAnalogAttackRelease = false; }),
+				FCanExecuteAction(),
+				FIsActionChecked::CreateSPLambda(this, [this]() { return !bIsAnalogAttackRelease; })
+			),
+			NAME_None,
+			EUserInterfaceActionType::ToggleButton);
+	}
+
 } // namespace AudioWidgets
+
+#undef LOCTEXT_NAMESPACE
diff --git a/Engine/Plugins/Runtime/AudioWidgets/Source/AudioWidgets/Private/SAudioSpectrumPlot.cpp b/Engine/Plugins/Runtime/AudioWidgets/Source/AudioWidgets/Private/SAudioSpectrumPlot.cpp
--- a/Engine/Plugins/Runtime/AudioWidgets/Source/AudioWidgets/Private/SAudioSpectrumPlot.cpp
+++ b/Engine/Plugins/Runtime/AudioWidgets/Source/AudioWidgets/Private/SAudioSpectrumPlot.cpp
@@ -4,10 +4,13 @@
 
 #include "Fonts/FontMeasure.h"
 #include "Framework/Application/SlateApplication.h"
+#include "Framework/MultiBox/MultiBoxBuilder.h"
 #include "SlateOptMacros.h"
 
 BEGIN_SLATE_FUNCTION_BUILD_OPTIMIZATION
 
+#define LOCTEXT_NAMESPACE "SAudioSpectrumPlot"
+
 /**
  * Helper class for drawing grid lines with text labels. Includes logic to avoid drawing overlapping labels if the grid lines are close together.
  */
@@ -189,6 +192,7 @@ bool FAudioSpectrumPlotGridAndLabelDrawingHelper::IsOverlappingPreviouslyDrawnLa
 };
 
 
+FName SAudioSpectrumPlot::ContextMenuExtensionHook("SpectrumPlotDisplayOptions");
 
 void SAudioSpectrumPlot::Construct(const FArguments& InArgs)
 {
@@ -207,9 +211,74 @@ void SAudioSpectrumPlot::Construct(const FArguments& InArgs)
 	GridColor = InArgs._GridColor;
 	AxisLabelColor = InArgs._AxisLabelColor;
 	SpectrumColor = InArgs._SpectrumColor;
+	bAllowContextMenu = InArgs._AllowContextMenu;
+	OnContextMenuOpening = InArgs._OnContextMenuOpening;
 	OnGetAudioSpectrumData = InArgs._OnGetAudioSpectrumData;
 }
 
+TSharedRef<const FExtensionBase> SAudioSpectrumPlot::AddContextMenuExtension(EExtensionHook::Position HookPosition, const TSharedPtr<FUICommandList>& CommandList, const FMenuExtensionDelegate& MenuExtensionDelegate)
+{
+	if (!ContextMenuExtender.IsValid())
+	{
+		ContextMenuExtender = MakeShared<FExtender>();
+	}
+
+	return ContextMenuExtender->AddMenuExtension(ContextMenuExtensionHook, HookPosition, CommandList, MenuExtensionDelegate);
+}
+
+void SAudioSpectrumPlot::RemoveContextMenuExtension(const TSharedRef<const FExtensionBase>& Extension)
+{
+	if (ensure(ContextMenuExtender.IsValid()))
+	{
+		ContextMenuExtender->RemoveExtension(Extension);
+	}
+}
+
+FReply SAudioSpectrumPlot::OnMouseButtonDown(const FGeometry& InMyGeometry, const FPointerEvent& InMouseEvent)
+{
+	if (!HasMouseCapture())
+	{
+		if (InMouseEvent.GetEffectingButton() == EKeys::RightMouseButton)
+		{
+			// Right clicking to summon context menu, but we'll do that on mouse-up.
+			return FReply::Handled().CaptureMouse(AsShared()).SetUserFocus(AsShared(), EFocusCause::Mouse);
+		}
+	}
+
+	return SCompoundWidget::OnMouseButtonDown(InMyGeometry, InMouseEvent);
+}
+
+FReply SAudioSpectrumPlot::OnMouseButtonUp(const FGeometry& InMyGeometry, const FPointerEvent& InMouseEvent)
+{
+	// The mouse must have been captured by mouse down before we'll process mouse ups
+	if (HasMouseCapture())
+	{
+		if (InMouseEvent.GetEffectingButton() == EKeys::RightMouseButton)
+		{
+			if (InMyGeometry.IsUnderLocation(InMouseEvent.GetScreenSpacePosition()) && bAllowContextMenu.Get())
+			{
+				TSharedPtr<SWidget> ContextMenu = OnContextMenuOpening.IsBound() ? OnContextMenuOpening.Execute() : BuildDefaultContextMenu();
+
+				if (ContextMenu.IsValid())
+				{
+					const FWidgetPath WidgetPath = (InMouseEvent.GetEventPath() != nullptr) ? *InMouseEvent.GetEventPath() : FWidgetPath();
+
+					FSlateApplication::Get().PushMenu(
+						AsShared(),
+						WidgetPath,
+						ContextMenu.ToSharedRef(),
+						InMouseEvent.GetScreenSpacePosition(),
+						FPopupTransitionEffect::ESlideDirection::ContextMenu);
+				}
+			}
+
+			return FReply::Handled().ReleaseMouseCapture();
+		}
+	}
+
+	return SCompoundWidget::OnMouseButtonUp(InMyGeometry, InMouseEvent);
+}
+
 int32 SAudioSpectrumPlot::OnPaint(const FPaintArgs& Args, const FGeometry& AllottedGeometry, const FSlateRect& MyCullingRect, FSlateWindowElementList& OutDrawElements, int32 LayerId, const FWidgetStyle& InWidgetStyle, bool bParentEnabled) const
 {
 	const FAudioSpectrumPlotScaleInfo ScaleInfo(AllottedGeometry.GetLocalSize(), FrequencyAxisScale.Get(), ViewMinFrequency.Get(), ViewMaxFrequency.Get(), ViewMinSoundLevel.Get(), ViewMaxSoundLevel.Get());
@@ -515,4 +584,117 @@ FLinearColor SAudioSpectrumPlot::GetSpectrumColor(const FWidgetStyle& InWidgetSt
 	const FSlateColor& SlateColor = (SpectrumColor.Get() != FSlateColor::UseStyle()) ? SpectrumColor.Get() : Style->SpectrumColor;
 	return SlateColor.GetColor(InWidgetStyle) * InWidgetStyle.GetColorAndOpacityTint();
 }
+
+TSharedRef<SWidget> SAudioSpectrumPlot::BuildDefaultContextMenu()
+{
+	constexpr bool bShouldCloseWindowAfterMenuSelection = true;
+	FMenuBuilder MenuBuilder(bShouldCloseWindowAfterMenuSelection, nullptr, ContextMenuExtender);
+
+	MenuBuilder.BeginSection(ContextMenuExtensionHook, LOCTEXT("DisplayOptions", "Display Options"));
+
+	if (!FrequencyAxisPixelBucketMode.IsBound())
+	{
+		MenuBuilder.AddSubMenu(
+			LOCTEXT("FrequencyAxisPixelBucketMode", "Pixel Plot Mode"),
+			FText(),
+			FNewMenuDelegate::CreateSP(this, &SAudioSpectrumPlot::BuildFrequencyAxisPixelBucketModeSubMenu));
+	}
+
+	if (!FrequencyAxisScale.IsBound())
+	{
+		MenuBuilder.AddSubMenu(
+			LOCTEXT("FrequencyAxisScale", "Frequency Scale"),
+			FText(),
+			FNewMenuDelegate::CreateSP(this, &SAudioSpectrumPlot::BuildFrequencyAxisScaleSubMenu));
+	}
+
+	if (!bDisplayFrequencyAxisLabels.IsBound())
+	{
+		MenuBuilder.AddMenuEntry(
+			LOCTEXT("DisplayFrequencyAxisLabels", "Display Frequency Axis Labels"),
+			FText(),
+			FSlateIcon(),
+			FUIAction(
+				FExecuteAction::CreateSPLambda(this, [this]() { bDisplayFrequencyAxisLabels = !bDisplayFrequencyAxisLabels.Get(); }),
+				FCanExecuteAction(),
+				FIsActionChecked::CreateSPLambda(this, [this]() { return bDisplayFrequencyAxisLabels.Get(); })
+			),
+			NAME_None,
+			EUserInterfaceActionType::ToggleButton);
+	}
+
+	if (!bDisplaySoundLevelAxisLabels.IsBound())
+	{
+		MenuBuilder.AddMenuEntry(
+			LOCTEXT("DisplaySoundLevelAxisLabels", "Display Sound Level Axis Labels"),
+			FText(),
+			FSlateIcon(),
+			FUIAction(
+				FExecuteAction::CreateSPLambda(this, [this]() { bDisplaySoundLevelAxisLabels = !bDisplaySoundLevelAxisLabels.Get(); }),
+				FCanExecuteAction(),
+				FIsActionChecked::CreateSPLambda(this, [this]() { return bDisplaySoundLevelAxisLabels.Get(); })
+			),
+			NAME_None,
+			EUserInterfaceActionType::ToggleButton);
+	}
+
+	MenuBuilder.EndSection();
+
+	return MenuBuilder.MakeWidget();
+}
+
+void SAudioSpectrumPlot::BuildFrequencyAxisScaleSubMenu(FMenuBuilder& SubMenu)
+{
+	const UEnum* EnumClass = StaticEnum<EAudioSpectrumPlotFrequencyAxisScale>();
+	const int32 NumEnumValues = EnumClass->NumEnums() - 1; // Exclude 'MAX' enum value.
+	for (int32 Index = 0; Index < NumEnumValues; Index++)
+	{
+		const auto EnumValue = static_cast<EAudioSpectrumPlotFrequencyAxisScale>(EnumClass->GetValueByIndex(Index));
+
+		SubMenu.AddMenuEntry(
+			EnumClass->GetDisplayNameTextByIndex(Index),
+#if WITH_EDITOR
+			EnumClass->GetToolTipTextByIndex(Index),
+#else
+			FText(),
+#endif
+			FSlateIcon(),
+			FUIAction(
+				FExecuteAction::CreateSP(this, &SAudioSpectrumPlot::SetFrequencyAxisScale, EnumValue),
+				FCanExecuteAction(),
+				FIsActionChecked::CreateSPLambda(this, [this, EnumValue]() { return (FrequencyAxisScale.Get() == EnumValue); })
+			),
+			NAME_None,
+			EUserInterfaceActionType::ToggleButton);
+	}
+}
+
+void SAudioSpectrumPlot::BuildFrequencyAxisPixelBucketModeSubMenu(FMenuBuilder& SubMenu)
+{
+	const UEnum* EnumClass = StaticEnum<EAudioSpectrumPlotFrequencyAxisPixelBucketMode>();
+	const int32 NumEnumValues = EnumClass->NumEnums() - 1; // Exclude 'MAX' enum value.
+	for (int32 Index = 0; Index < NumEnumValues; Index++)
+	{
+		const auto EnumValue = static_cast<EAudioSpectrumPlotFrequencyAxisPixelBucketMode>(EnumClass->GetValueByIndex(Index));
+
+		SubMenu.AddMenuEntry(
+			EnumClass->GetDisplayNameTextByIndex(Index),
+#if WITH_EDITOR
+			EnumClass->GetToolTipTextByIndex(Index),
+#else
+			FText(),
+#endif
+			FSlateIcon(),
+			FUIAction(
+				FExecuteAction::CreateSP(this, &SAudioSpectrumPlot::SetFrequencyAxisPixelBucketMode, EnumValue),
+				FCanExecuteAction(),
+				FIsActionChecked::CreateSPLambda(this, [this, EnumValue]() { return (FrequencyAxisPixelBucketMode.Get() == EnumValue); })
+			),
+			NAME_None,
+			EUserInterfaceActionType::ToggleButton);
+	}
+}
+
+#undef LOCTEXT_NAMESPACE
+
 END_SLATE_FUNCTION_BUILD_OPTIMIZATION
diff --git a/Engine/Plugins/Runtime/AudioWidgets/Source/AudioWidgets/Public/AudioSpectrumAnalyzer.h b/Engine/Plugins/Runtime/AudioWidgets/Source/AudioWidgets/Public/AudioSpectrumAnalyzer.h
--- a/Engine/Plugins/Runtime/AudioWidgets/Source/AudioWidgets/Public/AudioSpectrumAnalyzer.h
+++ b/Engine/Plugins/Runtime/AudioWidgets/Source/AudioWidgets/Public/AudioSpectrumAnalyzer.h
@@ -15,7 +15,7 @@ namespace AudioWidgets
 	 * Exponential time-smoothing is applied to the spectrum.
 	 * Can either create an Audio Bus to analyze, or analyze the given Bus.
 	 */
-	class AUDIOWIDGETS_API FAudioSpectrumAnalyzer
+	class AUDIOWIDGETS_API FAudioSpectrumAnalyzer : public TSharedFromThis<FAudioSpectrumAnalyzer>
 	{
 	public:
 		FAudioSpectrumAnalyzer(int32 InNumChannels, Audio::FDeviceId InAudioDeviceId, TObjectPtr<UAudioBus> InExternalAudioBus = nullptr);
@@ -30,6 +30,8 @@ namespace AudioWidgets
 	protected:
 		void OnConstantQResults(UConstantQAnalyzer* InSpectrumAnalyzer, int32 ChannelIndex, const TArray<FConstantQResults>& InSpectrumResultsArray);
 		FAudioPowerSpectrumData GetAudioSpectrumData() const;
+		void ExtendSpectrumPlotContextMenu(FMenuBuilder& MenuBuilder);
+		void BuildBallisticsSubMenu(FMenuBuilder& SubMenu);
 
 	private:
 		void Teardown();
@@ -54,11 +56,13 @@ namespace AudioWidgets
 
 		/** Slate widget for spectrum display */
 		TSharedPtr<SAudioSpectrumPlot> Widget;
+		TSharedPtr<const FExtensionBase> ContextMenuExtension;
 
 		bool bUseExternalAudioBus = false;
 
 		TOptional<float> PrevTimeStamp;
 		float AttackTimeMsec = 300.0f;
 		float ReleaseTimeMsec = 300.0f;
+		bool bIsAnalogAttackRelease = false;
 	};
 } // namespace AudioWidgets
diff --git a/Engine/Plugins/Runtime/AudioWidgets/Source/AudioWidgets/Public/AudioSpectrumPlotStyle.h b/Engine/Plugins/Runtime/AudioWidgets/Source/AudioWidgets/Public/AudioSpectrumPlotStyle.h
--- a/Engine/Plugins/Runtime/AudioWidgets/Source/AudioWidgets/Public/AudioSpectrumPlotStyle.h
+++ b/Engine/Plugins/Runtime/AudioWidgets/Source/AudioWidgets/Public/AudioSpectrumPlotStyle.h
@@ -38,7 +38,7 @@ struct AUDIOWIDGETS_API FAudioSpectrumPlotStyle : public FSlateWidgetStyle
 	FAudioSpectrumPlotStyle& SetAxisLabelColor(const FSlateColor& InAxisLabelColor) { AxisLabelColor = InAxisLabelColor; return *this; }
 
 	UPROPERTY(EditAnywhere, BlueprintReadWrite, Category = Appearance)
-	FSlateFontInfo AxisLabelFont = FStyleDefaults::GetFontInfo(8);
+	FSlateFontInfo AxisLabelFont = FStyleDefaults::GetFontInfo(5);
 	FAudioSpectrumPlotStyle& SetAxisLabelFont(const FSlateFontInfo& InAxisLabelFont) { AxisLabelFont = InAxisLabelFont; return *this; }
 
 	UPROPERTY(EditAnywhere, BlueprintReadWrite, Category = Appearance)
diff --git a/Engine/Plugins/Runtime/AudioWidgets/Source/AudioWidgets/Public/SAudioSpectrumPlot.h b/Engine/Plugins/Runtime/AudioWidgets/Source/AudioWidgets/Public/SAudioSpectrumPlot.h
--- a/Engine/Plugins/Runtime/AudioWidgets/Source/AudioWidgets/Public/SAudioSpectrumPlot.h
+++ b/Engine/Plugins/Runtime/AudioWidgets/Source/AudioWidgets/Public/SAudioSpectrumPlot.h
@@ -3,6 +3,8 @@
 #pragma once
 
 #include "AudioSpectrumPlotStyle.h"
+#include "Framework/MultiBox/MultiBoxExtender.h"
+#include "Framework/SlateDelegates.h"
 #include "Styling/ISlateStyle.h"
 #include "Styling/SlateWidgetStyleAsset.h"
 #include "Widgets/SCompoundWidget.h"
@@ -125,6 +127,7 @@ public:
 		, _GridColor(FSlateColor::UseStyle())
 		, _AxisLabelColor(FSlateColor::UseStyle())
 		, _SpectrumColor(FSlateColor::UseStyle())
+		, _AllowContextMenu(true)
 	{}
 		SLATE_STYLE_ARGUMENT(FAudioSpectrumPlotStyle, Style)
 		SLATE_ATTRIBUTE(float, ViewMinFrequency)
@@ -139,6 +142,8 @@ public:
 		SLATE_ATTRIBUTE(FSlateColor, GridColor)
 		SLATE_ATTRIBUTE(FSlateColor, AxisLabelColor)
 		SLATE_ATTRIBUTE(FSlateColor, SpectrumColor)
+		SLATE_ATTRIBUTE(bool, AllowContextMenu)
+		SLATE_EVENT(FOnContextMenuOpening, OnContextMenuOpening)
 		SLATE_EVENT(FGetAudioSpectrumData, OnGetAudioSpectrumData)
 	SLATE_END_ARGS()
 
@@ -153,6 +158,15 @@ public:
 	void SetDisplaySoundLevelAxisLabels(bool bInDisplaySoundLevelAxisLabels) { bDisplaySoundLevelAxisLabels = bInDisplaySoundLevelAxisLabels; }
 	void SetFrequencyAxisScale(EAudioSpectrumPlotFrequencyAxisScale InFrequencyAxisScale) { FrequencyAxisScale = InFrequencyAxisScale; }
 	void SetFrequencyAxisPixelBucketMode(EAudioSpectrumPlotFrequencyAxisPixelBucketMode InFrequencyAxisPixelBucketMode) { FrequencyAxisPixelBucketMode = InFrequencyAxisPixelBucketMode; }
+	void SetAllowContextMenu(bool bInAllowContextMenu) { bAllowContextMenu = bInAllowContextMenu; }
+
+	TSharedRef<const FExtensionBase> AddContextMenuExtension(EExtensionHook::Position HookPosition, const TSharedPtr<FUICommandList>& CommandList, const FMenuExtensionDelegate& MenuExtensionDelegate);
+	void RemoveContextMenuExtension(const TSharedRef<const FExtensionBase>& Extension);
+
+	// Begin SWidget overrides.
+	virtual FReply OnMouseButtonDown(const FGeometry& InMyGeometry, const FPointerEvent& InMouseEvent) override;
+	virtual FReply OnMouseButtonUp(const FGeometry& InMyGeometry, const FPointerEvent& InMouseEvent) override;
+	// End SWidget overrides.
 
 	void UnbindOnGetAudioSpectrumData() { OnGetAudioSpectrumData.Unbind(); }
 
@@ -179,6 +193,13 @@ private:
 	FLinearColor GetAxisLabelColor(const FWidgetStyle& InWidgetStyle) const;
 	FLinearColor GetSpectrumColor(const FWidgetStyle& InWidgetStyle) const;	
 
+	TSharedRef<SWidget> BuildDefaultContextMenu();
+	void BuildFrequencyAxisScaleSubMenu(FMenuBuilder& SubMenu);
+	void BuildFrequencyAxisPixelBucketModeSubMenu(FMenuBuilder& SubMenu);
+
+	static FName ContextMenuExtensionHook;
+	TSharedPtr<FExtender> ContextMenuExtender;
+
 	const FAudioSpectrumPlotStyle* Style;
 	TAttribute<float> ViewMinFrequency;
 	TAttribute<float> ViewMaxFrequency;
@@ -192,5 +213,7 @@ private:
 	TAttribute<FSlateColor> GridColor;
 	TAttribute<FSlateColor> AxisLabelColor;
 	TAttribute<FSlateColor> SpectrumColor;
+	TAttribute<bool> bAllowContextMenu;
+	FOnContextMenuOpening OnContextMenuOpening;
 	FGetAudioSpectrumData OnGetAudioSpectrumData;
 };
