Added ability to choose analyzer type from context menu in FAudioSpectrumAnalyzer.

diff --git a/Engine/Plugins/Runtime/AudioWidgets/Source/AudioWidgets/Private/AudioSpectrumAnalyzer.cpp b/Engine/Plugins/Runtime/AudioWidgets/Source/AudioWidgets/Private/AudioSpectrumAnalyzer.cpp
--- a/Engine/Plugins/Runtime/AudioWidgets/Source/AudioWidgets/Private/AudioSpectrumAnalyzer.cpp
+++ b/Engine/Plugins/Runtime/AudioWidgets/Source/AudioWidgets/Private/AudioSpectrumAnalyzer.cpp
@@ -2,6 +2,8 @@
 
 #include "AudioSpectrumAnalyzer.h"
 
+#include "AudioDevice.h"
+#include "AudioDeviceManager.h"
 #include "DSP/EnvelopeFollower.h"
 #include "Framework/MultiBox/MultiBoxBuilder.h"
 
@@ -10,13 +12,27 @@
 namespace AudioWidgets
 {
 	FAudioSpectrumAnalyzer::FAudioSpectrumAnalyzer(int32 InNumChannels, Audio::FDeviceId InAudioDeviceId, TObjectPtr<UAudioBus> InExternalAudioBus)
-		: Widget(SNew(SAudioSpectrumPlot)
+		: SpectrumAnalysisSettings(NewObject<USynesthesiaSpectrumAnalysisSettings>())
+		, ConstantQSettings(NewObject<UConstantQSettings>())
+		, Widget(SNew(SAudioSpectrumPlot)
 			.Clipping(EWidgetClipping::ClipToBounds)
 			.DisplayFrequencyAxisLabels(false)
 			.DisplaySoundLevelAxisLabels(false)
-			.FrequencyAxisPixelBucketMode_Lambda([]() { return EAudioSpectrumPlotFrequencyAxisPixelBucketMode::Average; }) // Binding this property has the effect of hiding its context menu entry (it isn't much use for the ConstantQ analyzer).
 			.OnGetAudioSpectrumData_Raw(this, &FAudioSpectrumAnalyzer::GetAudioSpectrumData))
 	{
+		SpectrumAnalysisSettings->SpectrumType = EAudioSpectrumType::PowerSpectrum;
+		SpectrumAnalysisSettings->FFTSize = EFFTSize::Max;
+		SpectrumAnalysisSettings->WindowType = EFFTWindowType::Blackman;
+
+		ConstantQSettings->SpectrumType = EAudioSpectrumType::PowerSpectrum;
+		ConstantQSettings->NumBandsPerOctave = 6.0f;
+		ConstantQSettings->NumBands = 61;
+		ConstantQSettings->StartingFrequencyHz = 20000.0f * FMath::Pow(0.5f, (ConstantQSettings->NumBands - 1) / ConstantQSettings->NumBandsPerOctave);
+		ConstantQSettings->FFTSize = EConstantQFFTSizeEnum::XXLarge;
+		ConstantQSettings->WindowType = EFFTWindowType::Blackman;
+		ConstantQSettings->bDownmixToMono = true;
+		ConstantQSettings->BandWidthStretch = 2.0f;
+
 		ContextMenuExtension = Widget->AddContextMenuExtension(EExtensionHook::Before, nullptr, FMenuExtensionDelegate::CreateRaw(this, &FAudioSpectrumAnalyzer::ExtendSpectrumPlotContextMenu));
 
 		Init(InNumChannels, InAudioDeviceId, InExternalAudioBus);
@@ -50,56 +66,102 @@ namespace AudioWidgets
 
 		Teardown();
 
-		Settings = TStrongObjectPtr(NewObject<UConstantQSettings>());
-		Settings->SpectrumType = EAudioSpectrumType::PowerSpectrum;
-		Settings->NumBandsPerOctave = 6.0f;
-		Settings->NumBands = 61;
-		Settings->StartingFrequencyHz = 20000.0f * FMath::Pow(0.5f, (Settings->NumBands - 1) / Settings->NumBandsPerOctave);
-		Settings->FFTSize = EConstantQFFTSizeEnum::XXLarge;
-		Settings->bDownmixToMono = true;
-		Settings->BandWidthStretch = 2.0f;
-
-		Analyzer = TStrongObjectPtr(NewObject<UConstantQAnalyzer>());
-		Analyzer->Settings = Settings.Get();
-
+		AudioDeviceId = InAudioDeviceId;
 		bUseExternalAudioBus = InExternalAudioBus != nullptr;
-
 		AudioBus = bUseExternalAudioBus ? TStrongObjectPtr(InExternalAudioBus.Get()) : TStrongObjectPtr(NewObject<UAudioBus>());
 		AudioBus->AudioBusChannels = EAudioBusChannels(InNumChannels - 1);
 
-		ResultsDelegateHandle = Analyzer->OnConstantQResultsNative.AddRaw(this, &FAudioSpectrumAnalyzer::OnConstantQResults);
+		SpectrumAnalyzer = TStrongObjectPtr(NewObject<USynesthesiaSpectrumAnalyzer>());
+		SpectrumAnalyzer->Settings = SpectrumAnalysisSettings.Get();
+		SpectrumResultsDelegateHandle = SpectrumAnalyzer->OnSpectrumResultsNative.AddRaw(this, &FAudioSpectrumAnalyzer::OnSpectrumResults);
 
-		Analyzer->StartAnalyzing(InAudioDeviceId, AudioBus.Get());
+		ConstantQAnalyzer = TStrongObjectPtr(NewObject<UConstantQAnalyzer>());
+		ConstantQAnalyzer->Settings = ConstantQSettings.Get();
+		ConstantQResultsDelegateHandle = ConstantQAnalyzer->OnConstantQResultsNative.AddRaw(this, &FAudioSpectrumAnalyzer::OnConstantQResults);
+
+		StartAnalyzing();
 	}
 
-	void FAudioSpectrumAnalyzer::OnConstantQResults(UConstantQAnalyzer* InSpectrumAnalyzer, int32 ChannelIndex, const TArray<FConstantQResults>& InSpectrumResultsArray)
+	void FAudioSpectrumAnalyzer::StartAnalyzing()
 	{
-		if (InSpectrumAnalyzer == Analyzer.Get())
+		switch (AnalyzerType)
 		{
-			for (const FConstantQResults& SpectrumResults : InSpectrumResultsArray)
+		case EAudioSpectrumAnalyzerType::FFT:
+			SpectrumAnalyzer->StartAnalyzing(AudioDeviceId, AudioBus.Get());
+			break;
+		case EAudioSpectrumAnalyzerType::CQT:
+			ConstantQAnalyzer->StartAnalyzing(AudioDeviceId, AudioBus.Get());
+			break;
+		default:
+			break;
+		}
+	}
+
+	void FAudioSpectrumAnalyzer::StopAnalyzing()
+	{
+		switch (AnalyzerType)
+		{
+		case EAudioSpectrumAnalyzerType::FFT:
+			SpectrumAnalyzer->StopAnalyzing();
+			break;
+		case EAudioSpectrumAnalyzerType::CQT:
+			ConstantQAnalyzer->StopAnalyzing();
+			break;
+		default:
+			break;
+		}
+	}
+
+	void FAudioSpectrumAnalyzer::OnSpectrumResults(USynesthesiaSpectrumAnalyzer* InSpectrumAnalyzer, int32 ChannelIndex, const TArray<FSynesthesiaSpectrumResults>& InSpectrumResultsArray)
+	{
+		if (AnalyzerType == EAudioSpectrumAnalyzerType::FFT && InSpectrumAnalyzer == SpectrumAnalyzer.Get())
+		{
+			for (const FSynesthesiaSpectrumResults& SpectrumResults : InSpectrumResultsArray)
 			{
 				if (PrevTimeStamp.IsSet() && SpectrumResults.TimeSeconds > PrevTimeStamp.GetValue())
 				{
-					// Calculate AR smoother coefficients:
-					const float DeltaT = (SpectrumResults.TimeSeconds - PrevTimeStamp.GetValue());
-					Audio::FAttackRelease AttackRelease(1.0f / DeltaT, AttackTimeMsec, ReleaseTimeMsec, bIsAnalogAttackRelease);
-
-					// Apply AR smoothing for each frequency:
-					check(SpectrumResults.SpectrumValues.Num() == ARSmoothedSquaredMagnitudes.Num());
-					for (int Index = 0; Index < SpectrumResults.SpectrumValues.Num(); Index++)
+					UpdateARSmoothing(SpectrumResults.TimeSeconds, SpectrumResults.SpectrumValues);
+				}
+				else
+				{
+					// Find samplerate:
+					float SampleRate = 48000.0f;
+					if (const FAudioDeviceManager* AudioDeviceManager = FAudioDeviceManager::Get())
 					{
-						const float OldValue = ARSmoothedSquaredMagnitudes[Index];
-						const float NewValue = SpectrumResults.SpectrumValues[Index];
-						const float ARSmootherCoefficient = (NewValue >= OldValue) ? AttackRelease.GetAttackTimeSamples() : AttackRelease.GetReleaseTimeSamples();
-						const float SmoothedValue = FMath::Lerp(NewValue, OldValue, ARSmootherCoefficient);
-						ARSmoothedSquaredMagnitudes[Index] = Audio::UnderflowClamp(SmoothedValue);
+						if (const FAudioDevice* AudioDevice = AudioDeviceManager->GetAudioDeviceRaw(AudioDeviceId))
+						{
+							SampleRate = AudioDevice->GetSampleRate();
+						}
 					}
+
+					// Init center frequencies:
+					CenterFrequencies.SetNumUninitialized(SpectrumAnalyzer->GetNumCenterFrequencies());
+					SpectrumAnalyzer->GetCenterFrequencies(SampleRate, CenterFrequencies);
+
+					// Init spectrum data:
+					ARSmoothedSquaredMagnitudes = SpectrumResults.SpectrumValues;
+				}
+
+				PrevTimeStamp = SpectrumResults.TimeSeconds;
+			}
+		}
+	}
+
+	void FAudioSpectrumAnalyzer::OnConstantQResults(UConstantQAnalyzer* InSpectrumAnalyzer, int32 ChannelIndex, const TArray<FConstantQResults>& InSpectrumResultsArray)
+	{
+		if (AnalyzerType == EAudioSpectrumAnalyzerType::CQT && InSpectrumAnalyzer == ConstantQAnalyzer.Get())
+		{
+			for (const FConstantQResults& SpectrumResults : InSpectrumResultsArray)
+			{
+				if (PrevTimeStamp.IsSet() && SpectrumResults.TimeSeconds > PrevTimeStamp.GetValue())
+				{
+					UpdateARSmoothing(SpectrumResults.TimeSeconds, SpectrumResults.SpectrumValues);
 				}
 				else
 				{
 					// Init center frequencies:
-					CenterFrequencies.SetNumUninitialized(Analyzer->GetNumCenterFrequencies());
-					Analyzer->GetCenterFrequencies(CenterFrequencies);
+					CenterFrequencies.SetNumUninitialized(ConstantQAnalyzer->GetNumCenterFrequencies());
+					ConstantQAnalyzer->GetCenterFrequencies(CenterFrequencies);
 
 					// Init spectrum data:
 					ARSmoothedSquaredMagnitudes = SpectrumResults.SpectrumValues;
@@ -110,26 +172,55 @@ namespace AudioWidgets
 		}
 	}
 
+	void FAudioSpectrumAnalyzer::UpdateARSmoothing(const float TimeStamp, TConstArrayView<float> SquaredMagnitudes)
+	{
+		// Calculate AR smoother coefficients:
+		const float DeltaT = (TimeStamp - PrevTimeStamp.GetValue());
+		Audio::FAttackRelease AttackRelease(1.0f / DeltaT, AttackTimeMsec, ReleaseTimeMsec, bIsAnalogAttackRelease);
+
+		// Apply AR smoothing for each frequency:
+		check(SquaredMagnitudes.Num() == ARSmoothedSquaredMagnitudes.Num());
+		for (int Index = 0; Index < SquaredMagnitudes.Num(); Index++)
+		{
+			const float OldValue = ARSmoothedSquaredMagnitudes[Index];
+			const float NewValue = SquaredMagnitudes[Index];
+			const float ARSmootherCoefficient = (NewValue >= OldValue) ? AttackRelease.GetAttackTimeSamples() : AttackRelease.GetReleaseTimeSamples();
+			const float SmoothedValue = FMath::Lerp(NewValue, OldValue, ARSmootherCoefficient);
+			ARSmoothedSquaredMagnitudes[Index] = Audio::UnderflowClamp(SmoothedValue);
+		}
+	}
+
 	void FAudioSpectrumAnalyzer::Teardown()
 	{
-		if (Analyzer.IsValid() && Analyzer->IsValidLowLevel())
+		if (SpectrumAnalyzer.IsValid() && SpectrumAnalyzer->IsValidLowLevel())
+		{
+			SpectrumAnalyzer->StopAnalyzing();
+			if (SpectrumResultsDelegateHandle.IsValid())
+			{
+				SpectrumAnalyzer->OnSpectrumResultsNative.Remove(SpectrumResultsDelegateHandle);
+			}
+
+			SpectrumAnalyzer.Reset();
+		}
+
+		if (ConstantQAnalyzer.IsValid() && ConstantQAnalyzer->IsValidLowLevel())
 		{
-			Analyzer->StopAnalyzing();
-			if (ResultsDelegateHandle.IsValid())
+			ConstantQAnalyzer->StopAnalyzing();
+			if (ConstantQResultsDelegateHandle.IsValid())
 			{
-				Analyzer->OnConstantQResultsNative.Remove(ResultsDelegateHandle);
+				ConstantQAnalyzer->OnConstantQResultsNative.Remove(ConstantQResultsDelegateHandle);
 			}
 			
-			Analyzer.Reset();
+			ConstantQAnalyzer.Reset();
 		}
 
-		ResultsDelegateHandle.Reset();
+		SpectrumResultsDelegateHandle.Reset();
+		ConstantQResultsDelegateHandle.Reset();
 		PrevTimeStamp.Reset();
 		CenterFrequencies.Empty();
 		ARSmoothedSquaredMagnitudes.Empty();
 
 		AudioBus.Reset();
-		Settings.Reset();
 
 		bUseExternalAudioBus = false;
 	}
@@ -147,6 +238,10 @@ namespace AudioWidgets
 			LOCTEXT("Ballistics", "Ballistics"),
 			FText(),
 			FNewMenuDelegate::CreateSP(this, &FAudioSpectrumAnalyzer::BuildBallisticsSubMenu));
+		MenuBuilder.AddSubMenu(
+			LOCTEXT("AnalyzerType", "Analyzer Type"),
+			FText(),
+			FNewMenuDelegate::CreateSP(this, &FAudioSpectrumAnalyzer::BuildAnalyzerTypeSubMenu));
 		MenuBuilder.EndSection();
 	}
 
@@ -176,6 +271,48 @@ namespace AudioWidgets
 			EUserInterfaceActionType::ToggleButton);
 	}
 
+	void FAudioSpectrumAnalyzer::BuildAnalyzerTypeSubMenu(FMenuBuilder& SubMenu)
+	{
+		const UEnum* EnumClass = StaticEnum<EAudioSpectrumAnalyzerType>();
+		const int32 NumEnumValues = EnumClass->NumEnums() - 1; // Exclude 'MAX' enum value.
+		for (int32 Index = 0; Index < NumEnumValues; Index++)
+		{
+			const auto EnumValue = static_cast<EAudioSpectrumAnalyzerType>(EnumClass->GetValueByIndex(Index));
+
+			SubMenu.AddMenuEntry(
+				EnumClass->GetDisplayNameTextByIndex(Index),
+#if WITH_EDITOR
+				EnumClass->GetToolTipTextByIndex(Index),
+#else
+				FText(),
+#endif
+				FSlateIcon(),
+				FUIAction(
+					FExecuteAction::CreateSP(this, &FAudioSpectrumAnalyzer::SetAnalyzerType, EnumValue),
+					FCanExecuteAction(),
+					FIsActionChecked::CreateSPLambda(this, [this, EnumValue]() { return (AnalyzerType == EnumValue); })
+				),
+				NAME_None,
+				EUserInterfaceActionType::ToggleButton);
+		}
+	}
+
+	void FAudioSpectrumAnalyzer::SetAnalyzerType(const EAudioSpectrumAnalyzerType InAnalyzerType)
+	{
+		if (AnalyzerType != InAnalyzerType)
+		{
+			StopAnalyzing();
+
+			AnalyzerType = InAnalyzerType;
+
+			PrevTimeStamp.Reset();
+			CenterFrequencies.Reset();
+			ARSmoothedSquaredMagnitudes.Reset();
+
+			StartAnalyzing();
+		}
+	}
+
 } // namespace AudioWidgets
 
 #undef LOCTEXT_NAMESPACE
diff --git a/Engine/Plugins/Runtime/AudioWidgets/Source/AudioWidgets/Public/AudioSpectrumAnalyzer.h b/Engine/Plugins/Runtime/AudioWidgets/Source/AudioWidgets/Public/AudioSpectrumAnalyzer.h
--- a/Engine/Plugins/Runtime/AudioWidgets/Source/AudioWidgets/Public/AudioSpectrumAnalyzer.h
+++ b/Engine/Plugins/Runtime/AudioWidgets/Source/AudioWidgets/Public/AudioSpectrumAnalyzer.h
@@ -4,10 +4,18 @@
 #include "ConstantQ.h"
 #include "SAudioSpectrumPlot.h"
 #include "Sound/AudioBus.h"
+#include "SynesthesiaSpectrumAnalysis.h"
 #include "UObject/StrongObjectPtr.h"
 
 class UWorld;
 
+UENUM(BlueprintType)
+enum class EAudioSpectrumAnalyzerType : uint8
+{
+	FFT UMETA(ToolTip = "Fast Fourier Transform"),
+	CQT UMETA(ToolTip = "Constant-Q Transform"),
+};
+
 namespace AudioWidgets
 {
 	/**
@@ -28,16 +36,27 @@ namespace AudioWidgets
 		void Init(int32 InNumChannels, Audio::FDeviceId InAudioDeviceId, TObjectPtr<UAudioBus> InExternalAudioBus = nullptr);
 
 	protected:
+		void StartAnalyzing();
+		void StopAnalyzing();
+
+		void OnSpectrumResults(USynesthesiaSpectrumAnalyzer* InSpectrumAnalyzer, int32 ChannelIndex, const TArray<FSynesthesiaSpectrumResults>& InSpectrumResultsArray);
 		void OnConstantQResults(UConstantQAnalyzer* InSpectrumAnalyzer, int32 ChannelIndex, const TArray<FConstantQResults>& InSpectrumResultsArray);
+		void UpdateARSmoothing(const float TimeStamp, TConstArrayView<float> SquaredMagnitudes);
+
 		FAudioPowerSpectrumData GetAudioSpectrumData() const;
+
 		void ExtendSpectrumPlotContextMenu(FMenuBuilder& MenuBuilder);
 		void BuildBallisticsSubMenu(FMenuBuilder& SubMenu);
+		void BuildAnalyzerTypeSubMenu(FMenuBuilder& SubMenu);
+
+		void SetAnalyzerType(const EAudioSpectrumAnalyzerType InAnalyzerType);
 
 	private:
 		void Teardown();
 
-		/** Audio analyzer object. */
-		TStrongObjectPtr<UConstantQAnalyzer> Analyzer;
+		/** Audio analyzer objects. */
+		TStrongObjectPtr<USynesthesiaSpectrumAnalyzer> SpectrumAnalyzer;
+		TStrongObjectPtr<UConstantQAnalyzer> ConstantQAnalyzer;
 
 		/** The audio bus used for analysis. */
 		TStrongObjectPtr<UAudioBus> AudioBus;
@@ -48,18 +67,23 @@ namespace AudioWidgets
 		/** Cached spectrum data, with AR smoothing applied. */
 		TArray<float> ARSmoothedSquaredMagnitudes;
 
-		/** Handle for results delegate for analyzer. */
-		FDelegateHandle ResultsDelegateHandle;
+		/** Handles for results delegate for analyzers. */
+		FDelegateHandle SpectrumResultsDelegateHandle;
+		FDelegateHandle ConstantQResultsDelegateHandle;
 
 		/** Analyzer settings. */
-		TStrongObjectPtr<UConstantQSettings> Settings;
+		TStrongObjectPtr<USynesthesiaSpectrumAnalysisSettings> SpectrumAnalysisSettings;
+		TStrongObjectPtr<UConstantQSettings> ConstantQSettings;
 
 		/** Slate widget for spectrum display */
 		TSharedPtr<SAudioSpectrumPlot> Widget;
 		TSharedPtr<const FExtensionBase> ContextMenuExtension;
 
+		Audio::FDeviceId AudioDeviceId = INDEX_NONE;
 		bool bUseExternalAudioBus = false;
 
+		EAudioSpectrumAnalyzerType AnalyzerType = EAudioSpectrumAnalyzerType::CQT;
+
 		TOptional<float> PrevTimeStamp;
 		float AttackTimeMsec = 300.0f;
 		float ReleaseTimeMsec = 300.0f;
