Added attributes to disable drawing grid lines in SAudioSpectrumPlot widget.

diff --git a/Engine/Plugins/Runtime/AudioWidgets/Source/AudioWidgets/Private/SAudioSpectrumPlot.cpp b/Engine/Plugins/Runtime/AudioWidgets/Source/AudioWidgets/Private/SAudioSpectrumPlot.cpp
--- a/Engine/Plugins/Runtime/AudioWidgets/Source/AudioWidgets/Private/SAudioSpectrumPlot.cpp
+++ b/Engine/Plugins/Runtime/AudioWidgets/Source/AudioWidgets/Private/SAudioSpectrumPlot.cpp
@@ -205,6 +205,8 @@ void SAudioSpectrumPlot::Construct(const FArguments& InArgs)
 	ViewMaxSoundLevel = InArgs._ViewMaxSoundLevel;
 	bDisplayFrequencyAxisLabels = InArgs._DisplayFrequencyAxisLabels;
 	bDisplaySoundLevelAxisLabels = InArgs._DisplaySoundLevelAxisLabels;
+	bDisplayFrequencyGridLines = InArgs._DisplayFrequencyGridLines;
+	bDisplaySoundLevelGridLines = InArgs._DisplaySoundLevelGridLines;
 	FrequencyAxisScale = InArgs._FrequencyAxisScale;
 	FrequencyAxisPixelBucketMode = InArgs._FrequencyAxisPixelBucketMode;
 	BackgroundColor = InArgs._BackgroundColor;
@@ -316,8 +318,15 @@ int32 SAudioSpectrumPlot::DrawGridAndLabels(const FGeometry& AllottedGeometry, F
 
 	FAudioSpectrumPlotGridAndLabelDrawingHelper GridAndLabelDrawingHelper(AllottedGeometry, OutDrawElements, ScaleInfo);
 
-	GridAndLabelDrawingHelper.DrawSoundLevelGridLines(LayerId, GridLineSoundLevels, LineColor);
-	GridAndLabelDrawingHelper.DrawFrequencyGridLines(LayerId, AllGridLineFrequencies, LineColor);
+	if (bDisplaySoundLevelGridLines.Get())
+	{
+		GridAndLabelDrawingHelper.DrawSoundLevelGridLines(LayerId, GridLineSoundLevels, LineColor);
+	}
+
+	if (bDisplayFrequencyGridLines.Get())
+	{
+		GridAndLabelDrawingHelper.DrawFrequencyGridLines(LayerId, AllGridLineFrequencies, LineColor);
+	}
 	LayerId++;
 
 	if (bDisplaySoundLevelAxisLabels.Get())
diff --git a/Engine/Plugins/Runtime/AudioWidgets/Source/AudioWidgets/Public/SAudioSpectrumPlot.h b/Engine/Plugins/Runtime/AudioWidgets/Source/AudioWidgets/Public/SAudioSpectrumPlot.h
--- a/Engine/Plugins/Runtime/AudioWidgets/Source/AudioWidgets/Public/SAudioSpectrumPlot.h
+++ b/Engine/Plugins/Runtime/AudioWidgets/Source/AudioWidgets/Public/SAudioSpectrumPlot.h
@@ -121,6 +121,8 @@ public:
 		, _ViewMaxSoundLevel(12.0f)
 		, _DisplayFrequencyAxisLabels(true)
 		, _DisplaySoundLevelAxisLabels(true)
+		, _DisplayFrequencyGridLines(true)
+		, _DisplaySoundLevelGridLines(true)
 		, _FrequencyAxisScale(EAudioSpectrumPlotFrequencyAxisScale::Logarithmic)
 		, _FrequencyAxisPixelBucketMode(EAudioSpectrumPlotFrequencyAxisPixelBucketMode::Average)
 		, _BackgroundColor(FSlateColor::UseStyle())
@@ -136,6 +138,8 @@ public:
 		SLATE_ATTRIBUTE(float, ViewMaxSoundLevel)
 		SLATE_ATTRIBUTE(bool, DisplayFrequencyAxisLabels)
 		SLATE_ATTRIBUTE(bool, DisplaySoundLevelAxisLabels)
+		SLATE_ATTRIBUTE(bool, DisplayFrequencyGridLines)
+		SLATE_ATTRIBUTE(bool, DisplaySoundLevelGridLines)
 		SLATE_ATTRIBUTE(EAudioSpectrumPlotFrequencyAxisScale, FrequencyAxisScale)
 		SLATE_ATTRIBUTE(EAudioSpectrumPlotFrequencyAxisPixelBucketMode, FrequencyAxisPixelBucketMode)
 		SLATE_ATTRIBUTE(FSlateColor, BackgroundColor)
@@ -156,6 +160,8 @@ public:
 	void SetViewMaxSoundLevel(float InViewMaxSoundLevel) { ViewMaxSoundLevel = InViewMaxSoundLevel; }
 	void SetDisplayFrequencyAxisLabels(bool bInDisplayFrequencyAxisLabels) { bDisplayFrequencyAxisLabels = bInDisplayFrequencyAxisLabels; }
 	void SetDisplaySoundLevelAxisLabels(bool bInDisplaySoundLevelAxisLabels) { bDisplaySoundLevelAxisLabels = bInDisplaySoundLevelAxisLabels; }
+	void SetDisplayFrequencyGridLines(bool bInDisplayFrequencyGridLines) { bDisplayFrequencyGridLines = bInDisplayFrequencyGridLines; }
+	void SetDisplaySoundLevelGridLines(bool bInDisplaySoundLevelGridLines) { bDisplaySoundLevelGridLines = bInDisplaySoundLevelGridLines; }
 	void SetFrequencyAxisScale(EAudioSpectrumPlotFrequencyAxisScale InFrequencyAxisScale) { FrequencyAxisScale = InFrequencyAxisScale; }
 	void SetFrequencyAxisPixelBucketMode(EAudioSpectrumPlotFrequencyAxisPixelBucketMode InFrequencyAxisPixelBucketMode) { FrequencyAxisPixelBucketMode = InFrequencyAxisPixelBucketMode; }
 	void SetAllowContextMenu(bool bInAllowContextMenu) { bAllowContextMenu = bInAllowContextMenu; }
@@ -207,6 +213,8 @@ private:
 	TAttribute<float> ViewMaxSoundLevel;
 	TAttribute<bool> bDisplayFrequencyAxisLabels;
 	TAttribute<bool> bDisplaySoundLevelAxisLabels;
+	TAttribute<bool> bDisplayFrequencyGridLines;
+	TAttribute<bool> bDisplaySoundLevelGridLines;
 	TAttribute<EAudioSpectrumPlotFrequencyAxisScale> FrequencyAxisScale;
 	TAttribute<EAudioSpectrumPlotFrequencyAxisPixelBucketMode> FrequencyAxisPixelBucketMode;
 	TAttribute<FSlateColor> BackgroundColor;
