Added FFT Size submenu to FAudioSpectrumAnalyzer context menu extension.

diff --git a/Engine/Plugins/Runtime/AudioWidgets/Source/AudioWidgets/Private/AudioSpectrumAnalyzer.cpp b/Engine/Plugins/Runtime/AudioWidgets/Source/AudioWidgets/Private/AudioSpectrumAnalyzer.cpp
--- a/Engine/Plugins/Runtime/AudioWidgets/Source/AudioWidgets/Private/AudioSpectrumAnalyzer.cpp
+++ b/Engine/Plugins/Runtime/AudioWidgets/Source/AudioWidgets/Private/AudioSpectrumAnalyzer.cpp
@@ -23,6 +23,7 @@ namespace AudioWidgets
 		SpectrumAnalysisSettings->SpectrumType = EAudioSpectrumType::PowerSpectrum;
 		SpectrumAnalysisSettings->FFTSize = EFFTSize::Max;
 		SpectrumAnalysisSettings->WindowType = EFFTWindowType::Blackman;
+		SpectrumAnalysisSettings->bDownmixToMono = true;
 
 		ConstantQSettings->SpectrumType = EAudioSpectrumType::PowerSpectrum;
 		ConstantQSettings->NumBandsPerOctave = 6.0f;
@@ -71,13 +72,8 @@ namespace AudioWidgets
 		AudioBus = bUseExternalAudioBus ? TStrongObjectPtr(InExternalAudioBus.Get()) : TStrongObjectPtr(NewObject<UAudioBus>());
 		AudioBus->AudioBusChannels = EAudioBusChannels(InNumChannels - 1);
 
-		SpectrumAnalyzer = TStrongObjectPtr(NewObject<USynesthesiaSpectrumAnalyzer>());
-		SpectrumAnalyzer->Settings = SpectrumAnalysisSettings.Get();
-		SpectrumResultsDelegateHandle = SpectrumAnalyzer->OnSpectrumResultsNative.AddRaw(this, &FAudioSpectrumAnalyzer::OnSpectrumResults);
-
-		ConstantQAnalyzer = TStrongObjectPtr(NewObject<UConstantQAnalyzer>());
-		ConstantQAnalyzer->Settings = ConstantQSettings.Get();
-		ConstantQResultsDelegateHandle = ConstantQAnalyzer->OnConstantQResultsNative.AddRaw(this, &FAudioSpectrumAnalyzer::OnConstantQResults);
+		CreateSynesthesiaSpectrumAnalyzer();
+		CreateConstantQAnalyzer();
 
 		StartAnalyzing();
 	}
@@ -194,34 +190,29 @@ namespace AudioWidgets
 	{
 		if (SpectrumAnalyzer.IsValid() && SpectrumAnalyzer->IsValidLowLevel())
 		{
-			SpectrumAnalyzer->StopAnalyzing();
-			if (SpectrumResultsDelegateHandle.IsValid())
+			if (AnalyzerType == EAudioSpectrumAnalyzerType::FFT)
 			{
-				SpectrumAnalyzer->OnSpectrumResultsNative.Remove(SpectrumResultsDelegateHandle);
+				SpectrumAnalyzer->StopAnalyzing();
 			}
 
-			SpectrumAnalyzer.Reset();
+			ReleaseSynesthesiaSpectrumAnalyzer();
 		}
 
 		if (ConstantQAnalyzer.IsValid() && ConstantQAnalyzer->IsValidLowLevel())
 		{
-			ConstantQAnalyzer->StopAnalyzing();
-			if (ConstantQResultsDelegateHandle.IsValid())
+			if (AnalyzerType == EAudioSpectrumAnalyzerType::CQT)
 			{
-				ConstantQAnalyzer->OnConstantQResultsNative.Remove(ConstantQResultsDelegateHandle);
+				ConstantQAnalyzer->StopAnalyzing();
 			}
-			
-			ConstantQAnalyzer.Reset();
+
+			ReleaseConstantQAnalyzer();
 		}
 
-		SpectrumResultsDelegateHandle.Reset();
-		ConstantQResultsDelegateHandle.Reset();
 		PrevTimeStamp.Reset();
 		CenterFrequencies.Empty();
 		ARSmoothedSquaredMagnitudes.Empty();
 
 		AudioBus.Reset();
-
 		bUseExternalAudioBus = false;
 	}
 
@@ -242,6 +233,10 @@ namespace AudioWidgets
 			LOCTEXT("AnalyzerType", "Analyzer Type"),
 			FText(),
 			FNewMenuDelegate::CreateSP(this, &FAudioSpectrumAnalyzer::BuildAnalyzerTypeSubMenu));
+		MenuBuilder.AddSubMenu(
+			LOCTEXT("FFTSize", "FFT Size"),
+			FText(),
+			FNewMenuDelegate::CreateSP(this, &FAudioSpectrumAnalyzer::BuildFFTSizeSubMenu));
 		MenuBuilder.EndSection();
 	}
 
@@ -297,6 +292,67 @@ namespace AudioWidgets
 		}
 	}
 
+	void FAudioSpectrumAnalyzer::BuildFFTSizeSubMenu(FMenuBuilder& SubMenu)
+	{
+		// There is a different FFTSize enum depending on the analyzer type.
+
+		if (AnalyzerType == EAudioSpectrumAnalyzerType::FFT)
+		{
+			const UEnum* EnumClass = StaticEnum<EFFTSize>();
+			const int32 NumEnumValues = EnumClass->NumEnums() - 1; // Exclude 'MAX' enum value.
+			for (int32 Index = 0; Index < NumEnumValues; Index++)
+			{
+				const auto EnumValue = static_cast<EFFTSize>(EnumClass->GetValueByIndex(Index));
+				if (EnumValue == EFFTSize::DefaultSize)
+				{
+					// Skip the duplicate 512 enum value 'DefaultSize'.
+					continue;
+				}
+
+				SubMenu.AddMenuEntry(
+					EnumClass->GetDisplayNameTextByIndex(Index),
+#if WITH_EDITOR
+					EnumClass->GetToolTipTextByIndex(Index),
+#else
+					FText(),
+#endif
+					FSlateIcon(),
+					FUIAction(
+						FExecuteAction::CreateSP(this, &FAudioSpectrumAnalyzer::SetSynesthesiaSpectrumAnalyzerFFTSize, EnumValue),
+						FCanExecuteAction(),
+						FIsActionChecked::CreateSPLambda(this, [this, EnumValue]() { return (SpectrumAnalysisSettings->FFTSize == EnumValue); })
+					),
+					NAME_None,
+					EUserInterfaceActionType::ToggleButton);
+			}
+		}
+		else if (AnalyzerType == EAudioSpectrumAnalyzerType::CQT)
+		{
+			const UEnum* EnumClass = StaticEnum<EConstantQFFTSizeEnum>();
+			const int32 NumEnumValues = EnumClass->NumEnums() - 1; // Exclude 'MAX' enum value.
+			for (int32 Index = 0; Index < NumEnumValues; Index++)
+			{
+				const auto EnumValue = static_cast<EConstantQFFTSizeEnum>(EnumClass->GetValueByIndex(Index));
+
+				SubMenu.AddMenuEntry(
+					EnumClass->GetDisplayNameTextByIndex(Index),
+#if WITH_EDITOR
+					EnumClass->GetToolTipTextByIndex(Index),
+#else
+					FText(),
+#endif
+					FSlateIcon(),
+					FUIAction(
+						FExecuteAction::CreateSP(this, &FAudioSpectrumAnalyzer::SetConstantQAnalyzerFFTSize, EnumValue),
+						FCanExecuteAction(),
+						FIsActionChecked::CreateSPLambda(this, [this, EnumValue]() { return (ConstantQSettings->FFTSize == EnumValue); })
+					),
+					NAME_None,
+					EUserInterfaceActionType::ToggleButton);
+			}
+		}
+	}
+
 	void FAudioSpectrumAnalyzer::SetAnalyzerType(const EAudioSpectrumAnalyzerType InAnalyzerType)
 	{
 		if (AnalyzerType != InAnalyzerType)
@@ -313,6 +369,73 @@ namespace AudioWidgets
 		}
 	}
 
+	void FAudioSpectrumAnalyzer::SetSynesthesiaSpectrumAnalyzerFFTSize(const EFFTSize FFTSize)
+	{
+		if (SpectrumAnalysisSettings->FFTSize != FFTSize)
+		{
+			StopAnalyzing();
+			ReleaseSynesthesiaSpectrumAnalyzer();
+
+			SpectrumAnalysisSettings->FFTSize = FFTSize;
+
+			CreateSynesthesiaSpectrumAnalyzer();
+			StartAnalyzing();
+		}
+	}
+
+	void FAudioSpectrumAnalyzer::SetConstantQAnalyzerFFTSize(const EConstantQFFTSizeEnum FFTSize)
+	{
+		if (ConstantQSettings->FFTSize != FFTSize)
+		{
+			StopAnalyzing();
+			ReleaseConstantQAnalyzer();
+
+			ConstantQSettings->FFTSize = FFTSize;
+
+			CreateConstantQAnalyzer();
+			StartAnalyzing();
+		}
+	}
+
+	void FAudioSpectrumAnalyzer::CreateSynesthesiaSpectrumAnalyzer()
+	{
+		ensure(!SpectrumAnalyzer.IsValid());
+		ensure(!SpectrumResultsDelegateHandle.IsValid());
+
+		SpectrumAnalyzer = TStrongObjectPtr(NewObject<USynesthesiaSpectrumAnalyzer>());
+		SpectrumAnalyzer->Settings = SpectrumAnalysisSettings.Get();
+		SpectrumResultsDelegateHandle = SpectrumAnalyzer->OnSpectrumResultsNative.AddRaw(this, &FAudioSpectrumAnalyzer::OnSpectrumResults);
+	}
+
+	void FAudioSpectrumAnalyzer::ReleaseSynesthesiaSpectrumAnalyzer()
+	{
+		if (ensure(SpectrumAnalyzer.IsValid() && SpectrumResultsDelegateHandle.IsValid()))
+		{
+			SpectrumAnalyzer->OnSpectrumResultsNative.Remove(SpectrumResultsDelegateHandle);
+		}
+
+		SpectrumResultsDelegateHandle.Reset();
+		SpectrumAnalyzer.Reset();
+	}
+
+	void FAudioSpectrumAnalyzer::CreateConstantQAnalyzer()
+	{
+		ConstantQAnalyzer = TStrongObjectPtr(NewObject<UConstantQAnalyzer>());
+		ConstantQAnalyzer->Settings = ConstantQSettings.Get();
+		ConstantQResultsDelegateHandle = ConstantQAnalyzer->OnConstantQResultsNative.AddRaw(this, &FAudioSpectrumAnalyzer::OnConstantQResults);
+	}
+
+	void FAudioSpectrumAnalyzer::ReleaseConstantQAnalyzer()
+	{
+		if (ensure(ConstantQAnalyzer.IsValid() && ConstantQResultsDelegateHandle.IsValid()))
+		{
+			ConstantQAnalyzer->OnConstantQResultsNative.Remove(ConstantQResultsDelegateHandle);
+		}
+
+		ConstantQResultsDelegateHandle.Reset();
+		ConstantQAnalyzer.Reset();
+	}
+
 } // namespace AudioWidgets
 
 #undef LOCTEXT_NAMESPACE
diff --git a/Engine/Plugins/Runtime/AudioWidgets/Source/AudioWidgets/Public/AudioSpectrumAnalyzer.h b/Engine/Plugins/Runtime/AudioWidgets/Source/AudioWidgets/Public/AudioSpectrumAnalyzer.h
--- a/Engine/Plugins/Runtime/AudioWidgets/Source/AudioWidgets/Public/AudioSpectrumAnalyzer.h
+++ b/Engine/Plugins/Runtime/AudioWidgets/Source/AudioWidgets/Public/AudioSpectrumAnalyzer.h
@@ -48,10 +48,19 @@ namespace AudioWidgets
 		void ExtendSpectrumPlotContextMenu(FMenuBuilder& MenuBuilder);
 		void BuildBallisticsSubMenu(FMenuBuilder& SubMenu);
 		void BuildAnalyzerTypeSubMenu(FMenuBuilder& SubMenu);
+		void BuildFFTSizeSubMenu(FMenuBuilder& SubMenu);
 
 		void SetAnalyzerType(const EAudioSpectrumAnalyzerType InAnalyzerType);
+		void SetSynesthesiaSpectrumAnalyzerFFTSize(const EFFTSize FFTSize);
+		void SetConstantQAnalyzerFFTSize(const EConstantQFFTSizeEnum FFTSize);
 
 	private:
+		void CreateSynesthesiaSpectrumAnalyzer();
+		void ReleaseSynesthesiaSpectrumAnalyzer();
+			
+		void CreateConstantQAnalyzer();
+		void ReleaseConstantQAnalyzer();
+
 		void Teardown();
 
 		/** Audio analyzer objects. */
