Modifiable DrawSize for the AudioMaterialSlates, plus optional size information pass to material.

diff --git a/Engine/Plugins/Runtime/AudioWidgets/Source/AudioWidgets/Private/AudioMaterialSlate/AudioMaterialEnvelope.cpp b/Engine/Plugins/Runtime/AudioWidgets/Source/AudioWidgets/Private/AudioMaterialSlate/AudioMaterialEnvelope.cpp
--- a/Engine/Plugins/Runtime/AudioWidgets/Source/AudioWidgets/Private/AudioMaterialSlate/AudioMaterialEnvelope.cpp
+++ b/Engine/Plugins/Runtime/AudioWidgets/Source/AudioWidgets/Private/AudioMaterialSlate/AudioMaterialEnvelope.cpp
@@ -37,7 +37,7 @@ TSharedRef<SWidget> UAudioMaterialEnvelope::RebuildWidget()
 {
 	EnvelopeCurve = SNew(SAudioMaterialEnvelope)
 		.Owner(this)
-		.AudioMaterialMeterStyle(&WidgetStyle)
+		.AudioMaterialEnvelopeStyle(&WidgetStyle)
 		.EnvelopeSettings(&EnvelopeSettings);
 
 	return EnvelopeCurve.ToSharedRef();
diff --git a/Engine/Plugins/Runtime/AudioWidgets/Source/AudioWidgets/Private/AudioMaterialSlate/AudioMaterialSlateStyles.cpp b/Engine/Plugins/Runtime/AudioWidgets/Source/AudioWidgets/Private/AudioMaterialSlate/AudioMaterialSlateStyles.cpp
--- a/Engine/Plugins/Runtime/AudioWidgets/Source/AudioWidgets/Private/AudioMaterialSlate/AudioMaterialSlateStyles.cpp
+++ b/Engine/Plugins/Runtime/AudioWidgets/Source/AudioWidgets/Private/AudioMaterialSlate/AudioMaterialSlateStyles.cpp
@@ -3,6 +3,11 @@
 
 #include "AudioMaterialSlate/AudioMaterialSlateStyles.h"
 
+FAudioMaterialStyle::FAudioMaterialStyle()
+	:DesiredSize(32.f, 32.f)
+{
+}
+
 TObjectPtr<UMaterialInstanceDynamic> FAudioMaterialStyle::GetDynamicMaterial() const
 {
 	return DynamicMaterial;
@@ -19,13 +24,14 @@ void FAudioMaterialSliderStyle::GetResources(TArray<const FSlateBrush*>& OutBrus
 }
 
 FAudioMaterialButtonStyle::FAudioMaterialButtonStyle()
-:ButtonMainColor(FLinearColor(0.1f, 0.1f, 0.1f, 1.0f)),
-ButtonShadowColor(FLinearColor(0.0f, 0.0f, 0.0f, 1.0f)),
-ButtonAccentColor(FLinearColor(0.34f, 0.34f, 0.34f, 1.0f)),
-ButtonPressedMainColor(FLinearColor(0.95f, 0.95f, 0.95f, 1.0f)),
-ButtonShadowMainColor(FLinearColor(0.156f, 0.138f, 0.15f, 1.0f)),
-ButtonPressedOutlineColor(FLinearColor(0.156f, 0.138f, 0.15f, 1.0f))
+	:ButtonMainColor(FLinearColor(0.1f, 0.1f, 0.1f, 1.0f))
+	,ButtonShadowColor(FLinearColor(0.0f, 0.0f, 0.0f, 1.0f))
+	,ButtonAccentColor(FLinearColor(0.34f, 0.34f, 0.34f, 1.0f))
+	,ButtonPressedMainColor(FLinearColor(0.95f, 0.95f, 0.95f, 1.0f))
+	,ButtonShadowMainColor(FLinearColor(0.156f, 0.138f, 0.15f, 1.0f))
+	,ButtonPressedOutlineColor(FLinearColor(0.156f, 0.138f, 0.15f, 1.0f))
 {
+	DesiredSize = FVector2f(128.f, 128.f);
 }
 
 FAudioMaterialButtonStyle::~FAudioMaterialButtonStyle()
@@ -45,12 +51,13 @@ const FAudioMaterialButtonStyle& FAudioMaterialButtonStyle::GetDefault()
 }
 
 FAudioMaterialSliderStyle::FAudioMaterialSliderStyle()
-:BarMainColor(FLinearColor(1.0f, 1.0f, 1.0f, 0.0f)),
-BarShadowColor(FLinearColor(1.0f, 1.0f, 1.0f, 0.0f)),
-BarAccentColor(FLinearColor(1.0f, 1.0f, 1.0f, 0.0f)),
-HandleMainColor(FLinearColor(1.0f, 1.0f, 1.0f, 0.0f)),
-HandleOutlineColor(FLinearColor(1.0f, 1.0f, 1.0f, 0.0f))
+	:BarMainColor(FLinearColor(1.0f, 1.0f, 1.0f, 0.0f))
+	,BarShadowColor(FLinearColor(1.0f, 1.0f, 1.0f, 0.0f))
+	,BarAccentColor(FLinearColor(1.0f, 1.0f, 1.0f, 0.0f))
+	,HandleMainColor(FLinearColor(1.0f, 1.0f, 1.0f, 0.0f))
+	,HandleOutlineColor(FLinearColor(1.0f, 1.0f, 1.0f, 0.0f))
 {
+	DesiredSize = FVector2f(128.f, 512.f);
 }
 
 FAudioMaterialSliderStyle::~FAudioMaterialSliderStyle()
@@ -66,16 +73,17 @@ const FAudioMaterialSliderStyle& FAudioMaterialSliderStyle::GetDefault()
 }
 
 FAudioMaterialKnobStyle::FAudioMaterialKnobStyle()
-:KnobMainColor(FLinearColor(0.f, 0.f, 0.f, 1.0f)),
-KnobAccentColor(FLinearColor(.14f, .16f, .2f, 1.0f)),
-KnobIndicatorColor(FLinearColor(1.f, 0.f, 0.04f, 1.0f)),
-KnobBarColor(FLinearColor(.033f, .033f, .033f, 1.0f)),
-KnobBarShadowColor(FLinearColor(0.f, 0.f, 0.f, 1.0f)),
-KnobBarFillMinColor(FLinearColor(0.9f, 0.9f, 0.9f, 1.0f)),
-KnobBarFillMidColor(FLinearColor(0.9f, 0.9f, 0.9f, 1.0f)),
-KnobBarFillMaxColor(FLinearColor(0.9f, 0.9f, 0.9f, 1.0f)),
-KnobBarFillTintColor(FLinearColor(1.f, 1.f, 1.f, 1.0f))
+	:KnobMainColor(FLinearColor(0.f, 0.f, 0.f, 1.0f))
+	,KnobAccentColor(FLinearColor(.14f, .16f, .2f, 1.0f))
+	,KnobIndicatorColor(FLinearColor(1.f, 0.f, 0.04f, 1.0f))
+	,KnobBarColor(FLinearColor(.033f, .033f, .033f, 1.0f))
+	,KnobBarShadowColor(FLinearColor(0.f, 0.f, 0.f, 1.0f))
+	,KnobBarFillMinColor(FLinearColor(0.9f, 0.9f, 0.9f, 1.0f))
+	,KnobBarFillMidColor(FLinearColor(0.9f, 0.9f, 0.9f, 1.0f))
+	,KnobBarFillMaxColor(FLinearColor(0.9f, 0.9f, 0.9f, 1.0f))
+	,KnobBarFillTintColor(FLinearColor(1.f, 1.f, 1.f, 1.0f))
 {
+	DesiredSize = FVector2f(256.f,256.f);
 }
 
 FAudioMaterialKnobStyle::~FAudioMaterialKnobStyle()
@@ -96,11 +104,12 @@ void FAudioMaterialKnobStyle::GetResources(TArray<const FSlateBrush*>& OutBrushe
 }
 
 FAudioMaterialMeterStyle::FAudioMaterialMeterStyle()
-:MeterFillMinColor(FLinearColor(1.f, 1.f, 1.f, 1.f)),
-MeterFillMidColor(FLinearColor(1.f, 1.f, 1.f, 1.f)),
-MeterFillMaxColor(FLinearColor(1.f, 1.f, 1.f, 1.f)),
-MeterOffFillColor(FLinearColor(0.5f, 0.5f, 0.5f, 1.f))
+	:MeterFillMinColor(FLinearColor(1.f, 1.f, 1.f, 1.f))
+	,MeterFillMidColor(FLinearColor(1.f, 1.f, 1.f, 1.f))
+	,MeterFillMaxColor(FLinearColor(1.f, 1.f, 1.f, 1.f))
+	,MeterOffFillColor(FLinearColor(0.5f, 0.5f, 0.5f, 1.f))
 {
+	DesiredSize = FVector2f(32.f, 512.f);
 }
 
 FAudioMaterialMeterStyle::~FAudioMaterialMeterStyle()
@@ -122,10 +131,11 @@ void FAudioMaterialMeterStyle::GetResources(TArray<const FSlateBrush*>& OutBrush
 
 
 FAudioMaterialEnvelopeStyle::FAudioMaterialEnvelopeStyle()
-:CurveColor(FLinearColor(0.15f, 0.45f, 0.75f, 1.f)),
-BackgroundColor(FLinearColor(0.01f, 0.01f, 0.01f, 1.f)),
-OutlineColor(FLinearColor(0.05f, 0.05f, 0.05f, 1.f))
+	:CurveColor(FLinearColor(0.15f, 0.45f, 0.75f, 1.f))
+	,BackgroundColor(FLinearColor(0.01f, 0.01f, 0.01f, 1.f))
+	,OutlineColor(FLinearColor(0.05f, 0.05f, 0.05f, 1.f))
 {
+	DesiredSize = FVector2f(256.f, 256.f);
 }
 
 FAudioMaterialEnvelopeStyle::~FAudioMaterialEnvelopeStyle()
diff --git a/Engine/Plugins/Runtime/AudioWidgets/Source/AudioWidgets/Private/AudioMaterialSlate/SAudioMaterialButton.cpp b/Engine/Plugins/Runtime/AudioWidgets/Source/AudioWidgets/Private/AudioMaterialSlate/SAudioMaterialButton.cpp
--- a/Engine/Plugins/Runtime/AudioWidgets/Source/AudioWidgets/Private/AudioMaterialSlate/SAudioMaterialButton.cpp
+++ b/Engine/Plugins/Runtime/AudioWidgets/Source/AudioWidgets/Private/AudioMaterialSlate/SAudioMaterialButton.cpp
@@ -46,6 +46,9 @@ int32 SAudioMaterialButton::OnPaint(const FPaintArgs& Args, const FGeometry& All
 			DynamicMaterial->SetVectorParameterValue(FName("Color_2"), AudioMaterialButtonStyle->ButtonShadowMainColor);
 			DynamicMaterial->SetVectorParameterValue(FName("LedColor"), AudioMaterialButtonStyle->ButtonPressedOutlineColor);
 			DynamicMaterial->SetScalarParameterValue(FName("Click"), bIsPressedAttribute.Get());
+
+			DynamicMaterial->SetScalarParameterValue(FName("LocalWidth"), AllottedGeometry.GetLocalSize().X);
+			DynamicMaterial->SetScalarParameterValue(FName("LocalHeigth"), AllottedGeometry.GetLocalSize().Y);			
 		}
 
 		const bool bEnabled = ShouldBeEnabled(bParentEnabled);
@@ -63,8 +66,12 @@ int32 SAudioMaterialButton::OnPaint(const FPaintArgs& Args, const FGeometry& All
 
 FVector2D SAudioMaterialButton::ComputeDesiredSize(float) const
 {
-	FVector2D Vector = FVector2D(SlateWidth, SlateHeight);
-	return Vector;
+	if (AudioMaterialButtonStyle)
+	{
+		return FVector2D(AudioMaterialButtonStyle->DesiredSize);
+	}
+
+	return FVector2D::ZeroVector;
 }
 
 FReply SAudioMaterialButton::OnMouseButtonDown(const FGeometry& MyGeometry, const FPointerEvent& MouseEvent)
diff --git a/Engine/Plugins/Runtime/AudioWidgets/Source/AudioWidgets/Private/AudioMaterialSlate/SAudioMaterialEnvelope.cpp b/Engine/Plugins/Runtime/AudioWidgets/Source/AudioWidgets/Private/AudioMaterialSlate/SAudioMaterialEnvelope.cpp
--- a/Engine/Plugins/Runtime/AudioWidgets/Source/AudioWidgets/Private/AudioMaterialSlate/SAudioMaterialEnvelope.cpp
+++ b/Engine/Plugins/Runtime/AudioWidgets/Source/AudioWidgets/Private/AudioMaterialSlate/SAudioMaterialEnvelope.cpp
@@ -14,7 +14,7 @@ void SAudioMaterialEnvelope::Construct(const FArguments& InArgs)
 {
 	Owner = InArgs._Owner;
 	EnvelopeSettings = InArgs._EnvelopeSettings;
-	AudioMaterialMeterStyle = InArgs._AudioMaterialMeterStyle;
+	AudioMaterialEnvelopeStyle = InArgs._AudioMaterialEnvelopeStyle;
 
 	ApplyNewMaterial();
 }
@@ -23,9 +23,9 @@ END_SLATE_FUNCTION_BUILD_OPTIMIZATION
 
 int32 SAudioMaterialEnvelope::OnPaint(const FPaintArgs& Args, const FGeometry& AllottedGeometry, const FSlateRect& MyCullingRect, FSlateWindowElementList& OutDrawElements, int32 LayerId, const FWidgetStyle& InWidgetStyle, bool bParentEnabled) const
 {
-	if (AudioMaterialMeterStyle && EnvelopeSettings)
+	if (AudioMaterialEnvelopeStyle && EnvelopeSettings)
 	{
-		UMaterialInstanceDynamic* DynamicMaterial = AudioMaterialMeterStyle->GetDynamicMaterial();
+		UMaterialInstanceDynamic* DynamicMaterial = AudioMaterialEnvelopeStyle->GetDynamicMaterial();
 		if (IsValid(DynamicMaterial))
 		{
 
@@ -44,9 +44,12 @@ int32 SAudioMaterialEnvelope::OnPaint(const FPaintArgs& Args, const FGeometry& A
 				DynamicMaterial->SetScalarParameterValue(FName("S_Int"), EnvelopeSettings->SustainValue);
 			}
 
-			DynamicMaterial->SetVectorParameterValue(FName("MainColor"), AudioMaterialMeterStyle->CurveColor);
-			DynamicMaterial->SetVectorParameterValue(FName("BoxBG"), AudioMaterialMeterStyle->BackgroundColor);
-			DynamicMaterial->SetVectorParameterValue(FName("BoxOutline"), AudioMaterialMeterStyle->OutlineColor);
+			DynamicMaterial->SetVectorParameterValue(FName("MainColor"), AudioMaterialEnvelopeStyle->CurveColor);
+			DynamicMaterial->SetVectorParameterValue(FName("BoxBG"), AudioMaterialEnvelopeStyle->BackgroundColor);
+			DynamicMaterial->SetVectorParameterValue(FName("BoxOutline"), AudioMaterialEnvelopeStyle->OutlineColor);
+
+			DynamicMaterial->SetScalarParameterValue(FName("LocalWidth"), AllottedGeometry.GetLocalSize().X);
+			DynamicMaterial->SetScalarParameterValue(FName("LocalHeigth"), AllottedGeometry.GetLocalSize().Y);			
 		}
 
 		const bool bEnabled = ShouldBeEnabled(bParentEnabled);
@@ -64,14 +67,18 @@ int32 SAudioMaterialEnvelope::OnPaint(const FPaintArgs& Args, const FGeometry& A
 
 FVector2D SAudioMaterialEnvelope::ComputeDesiredSize(float) const
 {
-	FVector2D Vector = FVector2D(SlateWidth, SlateHeight);
-	return Vector;
+	if (AudioMaterialEnvelopeStyle)
+	{
+		return FVector2D(AudioMaterialEnvelopeStyle->DesiredSize);
+	}
+
+	return FVector2D::ZeroVector;
 }
 
 void SAudioMaterialEnvelope::ApplyNewMaterial()
 {
-	if (AudioMaterialMeterStyle)
+	if (AudioMaterialEnvelopeStyle)
 	{
-		AudioMaterialMeterStyle->CreateDynamicMaterial(Owner.Get());
+		AudioMaterialEnvelopeStyle->CreateDynamicMaterial(Owner.Get());
 	}
 }
diff --git a/Engine/Plugins/Runtime/AudioWidgets/Source/AudioWidgets/Private/AudioMaterialSlate/SAudioMaterialKnob.cpp b/Engine/Plugins/Runtime/AudioWidgets/Source/AudioWidgets/Private/AudioMaterialSlate/SAudioMaterialKnob.cpp
--- a/Engine/Plugins/Runtime/AudioWidgets/Source/AudioWidgets/Private/AudioMaterialSlate/SAudioMaterialKnob.cpp
+++ b/Engine/Plugins/Runtime/AudioWidgets/Source/AudioWidgets/Private/AudioMaterialSlate/SAudioMaterialKnob.cpp
@@ -56,6 +56,9 @@ int32 SAudioMaterialKnob::OnPaint(const FPaintArgs& Args, const FGeometry& Allot
 			DynamicMaterial->SetVectorParameterValue(FName("LedTint"), AudioMaterialKnobStyle->KnobBarFillTintColor);
 
 			DynamicMaterial->SetScalarParameterValue(FName("VALUE"), FMath::Clamp(KnobPercent, 0.f, 1.f));
+
+			DynamicMaterial->SetScalarParameterValue(FName("LocalWidth"), AllottedGeometry.GetLocalSize().X);
+			DynamicMaterial->SetScalarParameterValue(FName("LocalHeigth"), AllottedGeometry.GetLocalSize().Y);			
 		}
 
 		const bool bEnabled = ShouldBeEnabled(bParentEnabled);
@@ -73,8 +76,12 @@ int32 SAudioMaterialKnob::OnPaint(const FPaintArgs& Args, const FGeometry& Allot
 
 FVector2D SAudioMaterialKnob::ComputeDesiredSize(float) const
 {
-	FVector2D Vector = FVector2D(SlateWidth, SlateHeight);
-	return Vector;
+	if (AudioMaterialKnobStyle)
+	{
+		return FVector2D(AudioMaterialKnobStyle->DesiredSize);
+	}
+
+	return FVector2D::ZeroVector;
 }
 
 FReply SAudioMaterialKnob::OnMouseMove(const FGeometry& MyGeometry, const FPointerEvent& MouseEvent)
diff --git a/Engine/Plugins/Runtime/AudioWidgets/Source/AudioWidgets/Private/AudioMaterialSlate/SAudioMaterialMeter.cpp b/Engine/Plugins/Runtime/AudioWidgets/Source/AudioWidgets/Private/AudioMaterialSlate/SAudioMaterialMeter.cpp
--- a/Engine/Plugins/Runtime/AudioWidgets/Source/AudioWidgets/Private/AudioMaterialSlate/SAudioMaterialMeter.cpp
+++ b/Engine/Plugins/Runtime/AudioWidgets/Source/AudioWidgets/Private/AudioMaterialSlate/SAudioMaterialMeter.cpp
@@ -38,6 +38,9 @@ int32 SAudioMaterialMeter::OnPaint(const FPaintArgs& Args, const FGeometry& Allo
 
 			const float Value = ValueAttribute.Get();
 			DynamicMaterial->SetScalarParameterValue(FName("VALUE"), FMath::Clamp(Value, 0.f, 1.f));
+
+			DynamicMaterial->SetScalarParameterValue(FName("LocalWidth"), AllottedGeometry.GetLocalSize().X);
+			DynamicMaterial->SetScalarParameterValue(FName("LocalHeigth"), AllottedGeometry.GetLocalSize().Y);			
 		}
 
 		const bool bEnabled = ShouldBeEnabled(bParentEnabled);
@@ -55,8 +58,12 @@ int32 SAudioMaterialMeter::OnPaint(const FPaintArgs& Args, const FGeometry& Allo
 
 FVector2D SAudioMaterialMeter::ComputeDesiredSize(float) const
 {
-	FVector2D Vector = FVector2D(SlateWidth, SlateHeight);
-	return Vector;
+	if (AudioMaterialMeterStyle)
+	{
+		return FVector2D(AudioMaterialMeterStyle->DesiredSize);
+	}
+
+	return FVector2D::ZeroVector;
 }
 
 void SAudioMaterialMeter::SetValue(float InValueAttribute)
diff --git a/Engine/Plugins/Runtime/AudioWidgets/Source/AudioWidgets/Private/AudioMaterialSlate/SAudioMaterialSlider.cpp b/Engine/Plugins/Runtime/AudioWidgets/Source/AudioWidgets/Private/AudioMaterialSlate/SAudioMaterialSlider.cpp
--- a/Engine/Plugins/Runtime/AudioWidgets/Source/AudioWidgets/Private/AudioMaterialSlate/SAudioMaterialSlider.cpp
+++ b/Engine/Plugins/Runtime/AudioWidgets/Source/AudioWidgets/Private/AudioMaterialSlate/SAudioMaterialSlider.cpp
@@ -36,6 +36,9 @@ int32 SAudioMaterialSlider::OnPaint(const FPaintArgs& Args, const FGeometry& All
 			DynamicMaterial->SetVectorParameterValue(FName("ValueColor"), AudioMaterialSliderStyle->HandleMainColor);
 			DynamicMaterial->SetVectorParameterValue(FName("DotBevel1"), AudioMaterialSliderStyle->HandleOutlineColor);
 			DynamicMaterial->SetScalarParameterValue(FName("VALUE"), FMath::Clamp(Value, 0.f, 1.f));
+
+			DynamicMaterial->SetScalarParameterValue(FName("LocalWidth"), AllottedGeometry.GetLocalSize().X);
+			DynamicMaterial->SetScalarParameterValue(FName("LocalHeigth"), AllottedGeometry.GetLocalSize().Y);			
 		}
 
 		const bool bEnabled = ShouldBeEnabled(bParentEnabled);
@@ -53,8 +56,12 @@ int32 SAudioMaterialSlider::OnPaint(const FPaintArgs& Args, const FGeometry& All
 
 FVector2D SAudioMaterialSlider::ComputeDesiredSize(float) const
 {
-	FVector2D Vector = FVector2D(SlateWidth, SlateHeight);
-	return Vector;
+	if (AudioMaterialSliderStyle)
+	{
+		return FVector2D(AudioMaterialSliderStyle->DesiredSize);
+	}
+
+	return FVector2D::ZeroVector;
 }
 
 void SAudioMaterialSlider::SetValue(TAttribute<float> InValueAttribute)
diff --git a/Engine/Plugins/Runtime/AudioWidgets/Source/AudioWidgets/Public/AudioMaterialSlate/AudioMaterialButton.h b/Engine/Plugins/Runtime/AudioWidgets/Source/AudioWidgets/Public/AudioMaterialSlate/AudioMaterialButton.h
--- a/Engine/Plugins/Runtime/AudioWidgets/Source/AudioWidgets/Public/AudioMaterialSlate/AudioMaterialButton.h
+++ b/Engine/Plugins/Runtime/AudioWidgets/Source/AudioWidgets/Public/AudioMaterialSlate/AudioMaterialButton.h
@@ -27,7 +27,7 @@ class AUDIOWIDGETS_API UAudioMaterialButton : public UWidget
 public:
 
 	/** The button's style */
-	UPROPERTY(EditAnywhere, BlueprintReadWrite, Category = "Style", meta = (DisplayName = "Style"))
+	UPROPERTY(EditAnywhere, BlueprintReadWrite, Category = "Style", meta = (DisplayName = "Style", ShowOnlyInnerProperties))
 	FAudioMaterialButtonStyle WidgetStyle;
 
 public:
diff --git a/Engine/Plugins/Runtime/AudioWidgets/Source/AudioWidgets/Public/AudioMaterialSlate/AudioMaterialEnvelope.h b/Engine/Plugins/Runtime/AudioWidgets/Source/AudioWidgets/Public/AudioMaterialSlate/AudioMaterialEnvelope.h
--- a/Engine/Plugins/Runtime/AudioWidgets/Source/AudioWidgets/Public/AudioMaterialSlate/AudioMaterialEnvelope.h
+++ b/Engine/Plugins/Runtime/AudioWidgets/Source/AudioWidgets/Public/AudioMaterialSlate/AudioMaterialEnvelope.h
@@ -24,7 +24,7 @@ class AUDIOWIDGETS_API UAudioMaterialEnvelope : public UWidget
 public:
 
 	/** The Envelope's style */
-	UPROPERTY(EditAnywhere, BlueprintReadWrite, Category = "Style", meta = (DisplayName = "Style"))
+	UPROPERTY(EditAnywhere, BlueprintReadWrite, Category = "Style", meta = (DisplayName = "Style", ShowOnlyInnerProperties))
 	FAudioMaterialEnvelopeStyle WidgetStyle;
 
 	/**Envelope settings*/
diff --git a/Engine/Plugins/Runtime/AudioWidgets/Source/AudioWidgets/Public/AudioMaterialSlate/AudioMaterialKnob.h b/Engine/Plugins/Runtime/AudioWidgets/Source/AudioWidgets/Public/AudioMaterialSlate/AudioMaterialKnob.h
--- a/Engine/Plugins/Runtime/AudioWidgets/Source/AudioWidgets/Public/AudioMaterialSlate/AudioMaterialKnob.h
+++ b/Engine/Plugins/Runtime/AudioWidgets/Source/AudioWidgets/Public/AudioMaterialSlate/AudioMaterialKnob.h
@@ -27,7 +27,7 @@ class AUDIOWIDGETS_API UAudioMaterialKnob : public UWidget
 public:
 
 	/** The button's style */
-	UPROPERTY(EditAnywhere, BlueprintReadWrite, Category = "Style", meta = (DisplayName = "Style"))
+	UPROPERTY(EditAnywhere, BlueprintReadWrite, Category = "Style", meta = (DisplayName = "Style", ShowOnlyInnerProperties))
 	FAudioMaterialKnobStyle WidgetStyle;
 
 public:
diff --git a/Engine/Plugins/Runtime/AudioWidgets/Source/AudioWidgets/Public/AudioMaterialSlate/AudioMaterialMeter.h b/Engine/Plugins/Runtime/AudioWidgets/Source/AudioWidgets/Public/AudioMaterialSlate/AudioMaterialMeter.h
--- a/Engine/Plugins/Runtime/AudioWidgets/Source/AudioWidgets/Public/AudioMaterialSlate/AudioMaterialMeter.h
+++ b/Engine/Plugins/Runtime/AudioWidgets/Source/AudioWidgets/Public/AudioMaterialSlate/AudioMaterialMeter.h
@@ -24,7 +24,7 @@ class AUDIOWIDGETS_API UAudioMaterialMeter : public UWidget
 public:
 
 	/** The meter's style */
-	UPROPERTY(EditAnywhere, BlueprintReadWrite, Category = "Style", meta = (DisplayName = "Style"))
+	UPROPERTY(EditAnywhere, BlueprintReadWrite, Category = "Style", meta = (DisplayName = "Style", ShowOnlyInnerProperties))
 	FAudioMaterialMeterStyle WidgetStyle;
 
 public:
diff --git a/Engine/Plugins/Runtime/AudioWidgets/Source/AudioWidgets/Public/AudioMaterialSlate/AudioMaterialSlateStyles.h b/Engine/Plugins/Runtime/AudioWidgets/Source/AudioWidgets/Public/AudioMaterialSlate/AudioMaterialSlateStyles.h
--- a/Engine/Plugins/Runtime/AudioWidgets/Source/AudioWidgets/Public/AudioMaterialSlate/AudioMaterialSlateStyles.h
+++ b/Engine/Plugins/Runtime/AudioWidgets/Source/AudioWidgets/Public/AudioMaterialSlate/AudioMaterialSlateStyles.h
@@ -20,9 +20,17 @@ struct AUDIOWIDGETS_API FAudioMaterialStyle : public FSlateWidgetStyle
 {
 	GENERATED_USTRUCT_BODY()
 
+public:
+
+	FAudioMaterialStyle();
+
 	/** Material used to render the Slate */
-	UPROPERTY(EditAnywhere, BlueprintReadWrite, Category = Appearance)
-	UMaterialInterface* Material = nullptr;
+	UPROPERTY(EditAnywhere, BlueprintReadWrite, meta = (DisplayPriority = 0), Category = "Style")
+	TObjectPtr<UMaterialInterface> Material = nullptr;
+
+	/** Desired Draw size of the rendered material*/
+	UPROPERTY(EditAnywhere, BlueprintReadWrite, meta = (DisplayPriority = 1), Category = "Style")
+	FVector2f DesiredSize;
 
 public:
 
@@ -52,22 +60,22 @@ struct AUDIOWIDGETS_API FAudioMaterialButtonStyle : public FAudioMaterialStyle
 	virtual const FName GetTypeName() const override { return TypeName; };
 	static const FAudioMaterialButtonStyle& GetDefault();
 
-	UPROPERTY(EditAnywhere, BlueprintReadWrite, Category = Appearance)
+	UPROPERTY(EditAnywhere, BlueprintReadWrite, Category = "Style")
 	FLinearColor ButtonMainColor;
 
-	UPROPERTY(EditAnywhere, BlueprintReadWrite, Category = Appearance)
+	UPROPERTY(EditAnywhere, BlueprintReadWrite, Category = "Style")
 	FLinearColor ButtonShadowColor;
 
-	UPROPERTY(EditAnywhere, BlueprintReadWrite, Category = Appearance)
+	UPROPERTY(EditAnywhere, BlueprintReadWrite, Category = "Style")
 	FLinearColor ButtonAccentColor;	
 	
-	UPROPERTY(EditAnywhere, BlueprintReadWrite, Category = Appearance)
+	UPROPERTY(EditAnywhere, BlueprintReadWrite, Category = "Style")
 	FLinearColor ButtonPressedMainColor;	
 	
-	UPROPERTY(EditAnywhere, BlueprintReadWrite, Category = Appearance)
+	UPROPERTY(EditAnywhere, BlueprintReadWrite, Category = "Style")
 	FLinearColor ButtonShadowMainColor;	
 	
-	UPROPERTY(EditAnywhere, BlueprintReadWrite, Category = Appearance)
+	UPROPERTY(EditAnywhere, BlueprintReadWrite, Category = "Style")
 	FLinearColor ButtonPressedOutlineColor;
 
 };
@@ -89,19 +97,19 @@ struct AUDIOWIDGETS_API FAudioMaterialSliderStyle : public FAudioMaterialStyle
 	virtual const FName GetTypeName() const override { return TypeName; };
 	static const FAudioMaterialSliderStyle& GetDefault();
 
-	UPROPERTY(EditAnywhere, BlueprintReadWrite, Category = Appearance)
+	UPROPERTY(EditAnywhere, BlueprintReadWrite, Category = "Style")
 	FLinearColor BarMainColor;	
 	
-	UPROPERTY(EditAnywhere, BlueprintReadWrite, Category = Appearance)
+	UPROPERTY(EditAnywhere, BlueprintReadWrite, Category = "Style")
 	FLinearColor BarShadowColor;	
 	
-	UPROPERTY(EditAnywhere, BlueprintReadWrite, Category = Appearance)
+	UPROPERTY(EditAnywhere, BlueprintReadWrite, Category = "Style")
 	FLinearColor BarAccentColor;
 
-	UPROPERTY(EditAnywhere, BlueprintReadWrite, Category = Appearance)
+	UPROPERTY(EditAnywhere, BlueprintReadWrite, Category = "Style")
 	FLinearColor HandleMainColor;	
 	
-	UPROPERTY(EditAnywhere, BlueprintReadWrite, Category = Appearance)
+	UPROPERTY(EditAnywhere, BlueprintReadWrite, Category = "Style")
 	FLinearColor HandleOutlineColor;
 
 };
@@ -123,31 +131,31 @@ struct AUDIOWIDGETS_API FAudioMaterialKnobStyle : public FAudioMaterialStyle
 	virtual const FName GetTypeName() const override { return TypeName; };
 	static const FAudioMaterialKnobStyle& GetDefault();
 
-	UPROPERTY(EditAnywhere, BlueprintReadWrite, Category = Appearance)
+	UPROPERTY(EditAnywhere, BlueprintReadWrite, Category = "Style")
 	FLinearColor KnobMainColor;
 
-	UPROPERTY(EditAnywhere, BlueprintReadWrite, Category = Appearance)
+	UPROPERTY(EditAnywhere, BlueprintReadWrite, Category = "Style")
 	FLinearColor KnobAccentColor;
 	
-	UPROPERTY(EditAnywhere, BlueprintReadWrite, Category = Appearance)
+	UPROPERTY(EditAnywhere, BlueprintReadWrite, Category = "Style")
 	FLinearColor KnobIndicatorColor;
 	
-	UPROPERTY(EditAnywhere, BlueprintReadWrite, Category = Appearance)
+	UPROPERTY(EditAnywhere, BlueprintReadWrite, Category = "Style")
 	FLinearColor KnobBarColor;		
 	
-	UPROPERTY(EditAnywhere, BlueprintReadWrite, Category = Appearance)
+	UPROPERTY(EditAnywhere, BlueprintReadWrite, Category = "Style")
 	FLinearColor KnobBarShadowColor;	
 
-	UPROPERTY(EditAnywhere, BlueprintReadWrite, Category = Appearance)
+	UPROPERTY(EditAnywhere, BlueprintReadWrite, Category = "Style")
 	FLinearColor KnobBarFillMinColor;		
 	
-	UPROPERTY(EditAnywhere, BlueprintReadWrite, Category = Appearance)
+	UPROPERTY(EditAnywhere, BlueprintReadWrite, Category = "Style")
 	FLinearColor KnobBarFillMidColor;		
 	
-	UPROPERTY(EditAnywhere, BlueprintReadWrite, Category = Appearance)
+	UPROPERTY(EditAnywhere, BlueprintReadWrite, Category = "Style")
 	FLinearColor KnobBarFillMaxColor;	
 	
-	UPROPERTY(EditAnywhere, BlueprintReadWrite, Category = Appearance)
+	UPROPERTY(EditAnywhere, BlueprintReadWrite, Category = "Style")
 	FLinearColor KnobBarFillTintColor;		
 	
 };
@@ -169,16 +177,16 @@ struct AUDIOWIDGETS_API FAudioMaterialMeterStyle : public FAudioMaterialStyle
 	virtual const FName GetTypeName() const override { return TypeName; };
 	static const FAudioMaterialMeterStyle& GetDefault();
 
-	UPROPERTY(EditAnywhere, BlueprintReadWrite, Category = Appearance)
+	UPROPERTY(EditAnywhere, BlueprintReadWrite, Category = "Style")
 	FLinearColor MeterFillMinColor;
 
-	UPROPERTY(EditAnywhere, BlueprintReadWrite, Category = Appearance)
+	UPROPERTY(EditAnywhere, BlueprintReadWrite, Category = "Style")
 	FLinearColor MeterFillMidColor;
 
-	UPROPERTY(EditAnywhere, BlueprintReadWrite, Category = Appearance)
+	UPROPERTY(EditAnywhere, BlueprintReadWrite, Category = "Style")
 	FLinearColor MeterFillMaxColor;	
 	
-	UPROPERTY(EditAnywhere, BlueprintReadWrite, Category = Appearance)
+	UPROPERTY(EditAnywhere, BlueprintReadWrite, Category = "Style")
 	FLinearColor MeterOffFillColor;
 	
 };
@@ -200,13 +208,13 @@ struct AUDIOWIDGETS_API FAudioMaterialEnvelopeStyle : public FAudioMaterialStyle
 	virtual const FName GetTypeName() const override { return TypeName; };
 	static const FAudioMaterialEnvelopeStyle& GetDefault();
 
-	UPROPERTY(EditAnywhere, BlueprintReadWrite, Category = Appearance)
+	UPROPERTY(EditAnywhere, BlueprintReadWrite, Category = "Style")
 	FLinearColor CurveColor;	
 	
-	UPROPERTY(EditAnywhere, BlueprintReadWrite, Category = Appearance)
+	UPROPERTY(EditAnywhere, BlueprintReadWrite, Category = "Style")
 	FLinearColor BackgroundColor;	
 	
-	UPROPERTY(EditAnywhere, BlueprintReadWrite, Category = Appearance)
+	UPROPERTY(EditAnywhere, BlueprintReadWrite, Category = "Style")
 	FLinearColor OutlineColor;
 	
 };
diff --git a/Engine/Plugins/Runtime/AudioWidgets/Source/AudioWidgets/Public/AudioMaterialSlate/AudioMaterialSlider.h b/Engine/Plugins/Runtime/AudioWidgets/Source/AudioWidgets/Public/AudioMaterialSlate/AudioMaterialSlider.h
--- a/Engine/Plugins/Runtime/AudioWidgets/Source/AudioWidgets/Public/AudioMaterialSlate/AudioMaterialSlider.h
+++ b/Engine/Plugins/Runtime/AudioWidgets/Source/AudioWidgets/Public/AudioMaterialSlate/AudioMaterialSlider.h
@@ -25,7 +25,7 @@ class AUDIOWIDGETS_API UAudioMaterialSlider : public UWidget
 public:
 
 	/** The slider's style */
-	UPROPERTY(EditAnywhere, BlueprintReadWrite, Category = "Style", meta = (DisplayName = "Style"))
+	UPROPERTY(EditAnywhere, BlueprintReadWrite, Category = "Style", meta = (DisplayName = "Style", ShowOnlyInnerProperties))
 	FAudioMaterialSliderStyle WidgetStyle;
 
 public:
diff --git a/Engine/Plugins/Runtime/AudioWidgets/Source/AudioWidgets/Public/AudioMaterialSlate/SAudioMaterialButton.h b/Engine/Plugins/Runtime/AudioWidgets/Source/AudioWidgets/Public/AudioMaterialSlate/SAudioMaterialButton.h
--- a/Engine/Plugins/Runtime/AudioWidgets/Source/AudioWidgets/Public/AudioMaterialSlate/SAudioMaterialButton.h
+++ b/Engine/Plugins/Runtime/AudioWidgets/Source/AudioWidgets/Public/AudioMaterialSlate/SAudioMaterialButton.h
@@ -72,9 +72,4 @@ private:
 	//Current pressed state of this button
 	TAttribute<bool> bIsPressedAttribute = false;
 
-	//Width & height
-	//TODO: Currently hardcoded for testing, make more dynamic.
-	float SlateWidth = 128.f;
-	float SlateHeight = 128.f;
-
 };
diff --git a/Engine/Plugins/Runtime/AudioWidgets/Source/AudioWidgets/Public/AudioMaterialSlate/SAudioMaterialEnvelope.h b/Engine/Plugins/Runtime/AudioWidgets/Source/AudioWidgets/Public/AudioMaterialSlate/SAudioMaterialEnvelope.h
--- a/Engine/Plugins/Runtime/AudioWidgets/Source/AudioWidgets/Public/AudioMaterialSlate/SAudioMaterialEnvelope.h
+++ b/Engine/Plugins/Runtime/AudioWidgets/Source/AudioWidgets/Public/AudioMaterialSlate/SAudioMaterialEnvelope.h
@@ -26,7 +26,7 @@ public:
 	SLATE_ARGUMENT(FAudioMaterialEnvelopeSettings*, EnvelopeSettings)
 
 	/** The style used to draw the Envelope. */
-	SLATE_STYLE_ARGUMENT(FAudioMaterialEnvelopeStyle, AudioMaterialMeterStyle)
+	SLATE_STYLE_ARGUMENT(FAudioMaterialEnvelopeStyle, AudioMaterialEnvelopeStyle)
 
 	SLATE_END_ARGS()
 
@@ -48,14 +48,9 @@ private:
 	TWeakObjectPtr<UObject> Owner;
 
 	// Holds the style for the Slate
-	const FAudioMaterialEnvelopeStyle* AudioMaterialMeterStyle = nullptr;
+	const FAudioMaterialEnvelopeStyle* AudioMaterialEnvelopeStyle = nullptr;
 
 	//Holds the current Envelope settings
 	const FAudioMaterialEnvelopeSettings* EnvelopeSettings = nullptr;
 
-	//Width & height
-	//TODO: Currently hardcoded for testing, make more dynamic.
-	float SlateWidth = 256.f;
-	float SlateHeight = 256.f;
-
 };
diff --git a/Engine/Plugins/Runtime/AudioWidgets/Source/AudioWidgets/Public/AudioMaterialSlate/SAudioMaterialKnob.h b/Engine/Plugins/Runtime/AudioWidgets/Source/AudioWidgets/Public/AudioMaterialSlate/SAudioMaterialKnob.h
--- a/Engine/Plugins/Runtime/AudioWidgets/Source/AudioWidgets/Public/AudioMaterialSlate/SAudioMaterialKnob.h
+++ b/Engine/Plugins/Runtime/AudioWidgets/Source/AudioWidgets/Public/AudioMaterialSlate/SAudioMaterialKnob.h
@@ -81,9 +81,4 @@ private:
 	// the max pixels to go to min or max value (clamped to 0 or 1) in one drag period
 	int32 PixelDelta = 50;
 
-	//Width & height
-	//TODO: Currently hardcoded for testing, make more dynamic.
-	float SlateWidth = 256.f;
-	float SlateHeight = 256.f;
-
 };
diff --git a/Engine/Plugins/Runtime/AudioWidgets/Source/AudioWidgets/Public/AudioMaterialSlate/SAudioMaterialMeter.h b/Engine/Plugins/Runtime/AudioWidgets/Source/AudioWidgets/Public/AudioMaterialSlate/SAudioMaterialMeter.h
--- a/Engine/Plugins/Runtime/AudioWidgets/Source/AudioWidgets/Public/AudioMaterialSlate/SAudioMaterialMeter.h
+++ b/Engine/Plugins/Runtime/AudioWidgets/Source/AudioWidgets/Public/AudioMaterialSlate/SAudioMaterialMeter.h
@@ -74,9 +74,4 @@ private:
 	//Holds the current value
 	TAttribute<float> ValueAttribute = 0.f;
 
-	//Width & height
-	//TODO: Currently hardcoded for testing, make more dynamic.
-	float SlateWidth = 32.f;
-	float SlateHeight = 512.f;
-
 };
diff --git a/Engine/Plugins/Runtime/AudioWidgets/Source/AudioWidgets/Public/AudioMaterialSlate/SAudioMaterialSlider.h b/Engine/Plugins/Runtime/AudioWidgets/Source/AudioWidgets/Public/AudioMaterialSlate/SAudioMaterialSlider.h
--- a/Engine/Plugins/Runtime/AudioWidgets/Source/AudioWidgets/Public/AudioMaterialSlate/SAudioMaterialSlider.h
+++ b/Engine/Plugins/Runtime/AudioWidgets/Source/AudioWidgets/Public/AudioMaterialSlate/SAudioMaterialSlider.h
@@ -88,9 +88,4 @@ private:
 	//Holds the current value
 	TAttribute<float> ValueAttribute = 0.f;
 
-	//Width & height
-	//TODO: Currently hardcoded for testing, make more dynamic.
-	float SlateWidth = 128.f;
-	float SlateHeight = 512.f;
-
 };
