DefaultStyleSets for AudioMaterialSlates

#virtualized

diff --git a/Engine/Plugins/Runtime/AudioWidgets/Source/AudioWidgets/Private/AudioMaterialSlate/AudioMaterialButton.cpp b/Engine/Plugins/Runtime/AudioWidgets/Source/AudioWidgets/Private/AudioMaterialSlate/AudioMaterialButton.cpp
--- a/Engine/Plugins/Runtime/AudioWidgets/Source/AudioWidgets/Private/AudioMaterialSlate/AudioMaterialButton.cpp
+++ b/Engine/Plugins/Runtime/AudioWidgets/Source/AudioWidgets/Private/AudioMaterialSlate/AudioMaterialButton.cpp
@@ -2,9 +2,18 @@
 
 
 #include "AudioMaterialSlate/AudioMaterialButton.h"
+#include "AudioMaterialSlate/AudioMaterialSlateTypes.h"
 #include "AudioMaterialSlate/SAudioMaterialButton.h"
+#include "AudioWidgetsStyle.h"
 
 #define LOCTEXT_NAMESPACE "AudioWidgets"
+
+UAudioMaterialButton::UAudioMaterialButton()
+{
+	//get default style
+	WidgetStyle = FAudioWidgetsStyle::Get().GetWidgetStyle<FAudioMaterialButtonStyle>("AudioMaterialButton.Style");
+}
+
 #if WITH_EDITOR
 const FText UAudioMaterialButton::GetPaletteCategory()
 {
diff --git a/Engine/Plugins/Runtime/AudioWidgets/Source/AudioWidgets/Private/AudioMaterialSlate/AudioMaterialEnvelope.cpp b/Engine/Plugins/Runtime/AudioWidgets/Source/AudioWidgets/Private/AudioMaterialSlate/AudioMaterialEnvelope.cpp
--- a/Engine/Plugins/Runtime/AudioWidgets/Source/AudioWidgets/Private/AudioMaterialSlate/AudioMaterialEnvelope.cpp
+++ b/Engine/Plugins/Runtime/AudioWidgets/Source/AudioWidgets/Private/AudioMaterialSlate/AudioMaterialEnvelope.cpp
@@ -2,13 +2,20 @@
 
 
 #include "AudioMaterialSlate/AudioMaterialEnvelope.h"
-#include "AudioMaterialSlate/AudioMaterialSlateStyles.h"
+#include "AudioMaterialSlate/AudioMaterialSlateTypes.h"
 #include "AudioMaterialSlate/SAudioMaterialEnvelope.h"
+#include "AudioWidgetsStyle.h"
 #include "Widgets/SWeakWidget.h"
 #include "Components/AudioComponent.h"
 
 #define LOCTEXT_NAMESPACE "AudioWidgets"
 
+UAudioMaterialEnvelope::UAudioMaterialEnvelope()
+{
+	//get default style
+	WidgetStyle = FAudioWidgetsStyle::Get().GetWidgetStyle<FAudioMaterialEnvelopeStyle>("AudioMaterialEnvelope.Style");
+}
+
 #if WITH_EDITOR
 const FText UAudioMaterialEnvelope::GetPaletteCategory()
 {
diff --git a/Engine/Plugins/Runtime/AudioWidgets/Source/AudioWidgets/Private/AudioMaterialSlate/AudioMaterialKnob.cpp b/Engine/Plugins/Runtime/AudioWidgets/Source/AudioWidgets/Private/AudioMaterialSlate/AudioMaterialKnob.cpp
--- a/Engine/Plugins/Runtime/AudioWidgets/Source/AudioWidgets/Private/AudioMaterialSlate/AudioMaterialKnob.cpp
+++ b/Engine/Plugins/Runtime/AudioWidgets/Source/AudioWidgets/Private/AudioMaterialSlate/AudioMaterialKnob.cpp
@@ -2,9 +2,17 @@
 
 
 #include "AudioMaterialSlate/AudioMaterialKnob.h"
+#include "AudioMaterialSlate/AudioMaterialSlateTypes.h"
 #include "AudioMaterialSlate/SAudioMaterialKnob.h"
+#include "AudioWidgetsStyle.h"
 
 #define LOCTEXT_NAMESPACE "AudioWidgets"
+UAudioMaterialKnob::UAudioMaterialKnob()
+{
+	//get default style
+	WidgetStyle = FAudioWidgetsStyle::Get().GetWidgetStyle<FAudioMaterialKnobStyle>("AudioMaterialKnob.Style");	
+}
+
 #if WITH_EDITOR
 const FText UAudioMaterialKnob::GetPaletteCategory()
 {
diff --git a/Engine/Plugins/Runtime/AudioWidgets/Source/AudioWidgets/Private/AudioMaterialSlate/AudioMaterialMeter.cpp b/Engine/Plugins/Runtime/AudioWidgets/Source/AudioWidgets/Private/AudioMaterialSlate/AudioMaterialMeter.cpp
--- a/Engine/Plugins/Runtime/AudioWidgets/Source/AudioWidgets/Private/AudioMaterialSlate/AudioMaterialMeter.cpp
+++ b/Engine/Plugins/Runtime/AudioWidgets/Source/AudioWidgets/Private/AudioMaterialSlate/AudioMaterialMeter.cpp
@@ -2,13 +2,20 @@
 
 
 #include "AudioMaterialSlate/AudioMaterialMeter.h"
-#include "AudioMaterialSlate/AudioMaterialSlateStyles.h"
+#include "AudioMaterialSlate/AudioMaterialSlateTypes.h"
 #include "AudioMaterialSlate/SAudioMaterialMeter.h"
+#include "AudioWidgetsStyle.h"
 #include "Widgets/SWeakWidget.h"
 #include "Components/AudioComponent.h"
 
 #define LOCTEXT_NAMESPACE "AudioWidgets"
 
+UAudioMaterialMeter::UAudioMaterialMeter()
+{
+	//get default style
+	WidgetStyle = FAudioWidgetsStyle::Get().GetWidgetStyle<FAudioMaterialMeterStyle>("AudioMaterialMeter.Style");
+}
+
 #if WITH_EDITOR
 const FText UAudioMaterialMeter::GetPaletteCategory()
 {
diff --git a/Engine/Plugins/Runtime/AudioWidgets/Source/AudioWidgets/Private/AudioMaterialSlate/AudioMaterialSlateStyles.cpp b/Engine/Plugins/Runtime/AudioWidgets/Source/AudioWidgets/Private/AudioMaterialSlate/AudioMaterialSlateStyles.cpp
deleted file mode 100644
diff --git a/Engine/Plugins/Runtime/AudioWidgets/Source/AudioWidgets/Private/AudioMaterialSlate/AudioMaterialSlateTypes.cpp b/Engine/Plugins/Runtime/AudioWidgets/Source/AudioWidgets/Private/AudioMaterialSlate/AudioMaterialSlateTypes.cpp
new file mode 100644
--- /dev/null
+++ b/Engine/Plugins/Runtime/AudioWidgets/Source/AudioWidgets/Private/AudioMaterialSlate/AudioMaterialSlateTypes.cpp
@@ -0,0 +1,135 @@
+// Copyright Epic Games, Inc. All Rights Reserved.
+
+
+#include "AudioMaterialSlate/AudioMaterialSlateTypes.h"
+
+FAudioMaterialWidgetStyle::FAudioMaterialWidgetStyle()
+	:DesiredSize(32.f, 32.f)
+{
+}
+
+TObjectPtr<UMaterialInstanceDynamic> FAudioMaterialWidgetStyle::GetDynamicMaterial() const
+{
+	return DynamicMaterial;
+}
+
+void FAudioMaterialWidgetStyle::CreateDynamicMaterial(UObject* InOuter) const
+{
+	DynamicMaterial = UMaterialInstanceDynamic::Create(Material, InOuter);
+}
+
+FAudioMaterialButtonStyle::FAudioMaterialButtonStyle()
+	: ButtonMainColor(FLinearColor::Gray)
+	, ButtonShadowColor(FLinearColor::Black)
+	, ButtonAccentColor(FLinearColor::Gray)
+	, ButtonPressedMainColor(FLinearColor::White)
+	, ButtonPressedShadowColor(FLinearColor::Gray)
+	, ButtonPressedOutlineColor(FLinearColor::Blue)
+{
+	DesiredSize = FVector2f(128.f, 128.f);
+}
+
+const FName FAudioMaterialButtonStyle::TypeName(TEXT("FAudioMaterialButtonStyle"));
+
+void FAudioMaterialButtonStyle::GetResources(TArray<const FSlateBrush*>& OutBrushes) const
+{
+}
+
+const FAudioMaterialButtonStyle& FAudioMaterialButtonStyle::GetDefault()
+{
+	static FAudioMaterialButtonStyle Default;
+	return Default;
+}
+
+FAudioMaterialSliderStyle::FAudioMaterialSliderStyle()
+	: BarMainColor(FLinearColor::White)
+	, BarShadowColor(FLinearColor::White)
+	, BarAccentColor(FLinearColor::White)
+	, HandleMainColor(FLinearColor::White)
+	, HandleOutlineColor(FLinearColor::White)
+{
+	DesiredSize = FVector2f(128.f, 512.f);
+}
+
+const FName FAudioMaterialSliderStyle::TypeName(TEXT("FAudioMaterialSliderStyle"));
+
+void FAudioMaterialSliderStyle::GetResources(TArray<const FSlateBrush*>& OutBrushes) const
+{
+	// Add any brush resources here so that Slate can correctly atlas and reference them
+}
+
+const FAudioMaterialSliderStyle& FAudioMaterialSliderStyle::GetDefault()
+{
+	static FAudioMaterialSliderStyle Default;
+	return Default;
+}
+
+FAudioMaterialKnobStyle::FAudioMaterialKnobStyle()
+	: KnobMainColor(FLinearColor::Black)
+	, KnobAccentColor(FLinearColor::Gray)
+	, KnobIndicatorColor(FLinearColor::Red)
+	, KnobBarColor(FLinearColor::Gray)
+	, KnobBarShadowColor(FLinearColor::Gray)
+	, KnobBarFillMinColor(FLinearColor::White)
+	, KnobBarFillMidColor(FLinearColor::White)
+	, KnobBarFillMaxColor(FLinearColor::White)
+	, KnobBarFillTintColor(FLinearColor::White)
+{
+	DesiredSize = FVector2f(256.f,256.f);
+}
+
+const FName FAudioMaterialKnobStyle::TypeName(TEXT("FAudioMaterialKnobStyle"));
+
+const FAudioMaterialKnobStyle& FAudioMaterialKnobStyle::GetDefault()
+{
+	static FAudioMaterialKnobStyle Default;
+	return Default;
+}
+
+void FAudioMaterialKnobStyle::GetResources(TArray<const FSlateBrush*>& OutBrushes) const
+{
+	// Add any brush resources here so that Slate can correctly atlas and reference them
+}
+
+FAudioMaterialMeterStyle::FAudioMaterialMeterStyle()
+	:MeterFillMinColor(FLinearColor::White)
+	, MeterFillMidColor(FLinearColor::White)
+	, MeterFillMaxColor(FLinearColor::White)
+	, MeterOffFillColor(FLinearColor::Black)
+{
+	DesiredSize = FVector2f(32.f, 512.f);
+}
+
+const FName FAudioMaterialMeterStyle::TypeName(TEXT("FAudioMaterialMeterStyle"));
+
+const FAudioMaterialMeterStyle& FAudioMaterialMeterStyle::GetDefault()
+{
+	static FAudioMaterialMeterStyle Default;
+	return Default;
+}
+
+void FAudioMaterialMeterStyle::GetResources(TArray<const FSlateBrush*>& OutBrushes) const
+{
+	// Add any brush resources here so that Slate can correctly atlas and reference them
+}
+
+FAudioMaterialEnvelopeStyle::FAudioMaterialEnvelopeStyle()
+	: CurveColor(FLinearColor::White)
+	, BackgroundColor(FLinearColor::Black)
+	, OutlineColor(FLinearColor::Gray)
+{
+	DesiredSize = FVector2f(256.f, 256.f);
+}
+
+const FName FAudioMaterialEnvelopeStyle::TypeName(TEXT("FAudioMaterialEnvelopeStyle"));
+
+const FAudioMaterialEnvelopeStyle& FAudioMaterialEnvelopeStyle::GetDefault()
+{
+	static FAudioMaterialEnvelopeStyle Default;
+	return Default;
+}
+
+void FAudioMaterialEnvelopeStyle::GetResources(TArray<const FSlateBrush*>& OutBrushes) const
+{
+	// Add any brush resources here so that Slate can correctly atlas and reference them
+}
diff --git a/Engine/Plugins/Runtime/AudioWidgets/Source/AudioWidgets/Private/AudioMaterialSlate/AudioMaterialSlider.cpp b/Engine/Plugins/Runtime/AudioWidgets/Source/AudioWidgets/Private/AudioMaterialSlate/AudioMaterialSlider.cpp
--- a/Engine/Plugins/Runtime/AudioWidgets/Source/AudioWidgets/Private/AudioMaterialSlate/AudioMaterialSlider.cpp
+++ b/Engine/Plugins/Runtime/AudioWidgets/Source/AudioWidgets/Private/AudioMaterialSlate/AudioMaterialSlider.cpp
@@ -2,12 +2,19 @@
 
 
 #include "AudioMaterialSlate/AudioMaterialSlider.h"
-#include "AudioMaterialSlate/AudioMaterialSlateStyles.h"
+#include "AudioMaterialSlate/AudioMaterialSlateTypes.h"
 #include "AudioMaterialSlate/SAudioMaterialSlider.h"
+#include "AudioWidgetsStyle.h"
 #include "Widgets/SWeakWidget.h"
 
 #define LOCTEXT_NAMESPACE "AudioWidgets"
 
+UAudioMaterialSlider::UAudioMaterialSlider()
+{
+	//get default style
+	WidgetStyle = FAudioWidgetsStyle::Get().GetWidgetStyle<FAudioMaterialSliderStyle>("AudioMaterialSlider.Style");
+}
+
 #if WITH_EDITOR
 const FText UAudioMaterialSlider::GetPaletteCategory()
 {
diff --git a/Engine/Plugins/Runtime/AudioWidgets/Source/AudioWidgets/Private/AudioMaterialSlate/SAudioMaterialButton.cpp b/Engine/Plugins/Runtime/AudioWidgets/Source/AudioWidgets/Private/AudioMaterialSlate/SAudioMaterialButton.cpp
--- a/Engine/Plugins/Runtime/AudioWidgets/Source/AudioWidgets/Private/AudioMaterialSlate/SAudioMaterialButton.cpp
+++ b/Engine/Plugins/Runtime/AudioWidgets/Source/AudioWidgets/Private/AudioMaterialSlate/SAudioMaterialButton.cpp
@@ -2,6 +2,7 @@
 
 
 #include "AudioMaterialSlate/SAudioMaterialButton.h"
+#include "AudioMaterialSlate/AudioMaterialSlateTypes.h"
 #include "SlateOptMacros.h"
 #include "Components/Widget.h"
 
@@ -43,7 +44,7 @@ int32 SAudioMaterialButton::OnPaint(const FPaintArgs& Args, const FGeometry& All
 			DynamicMaterial->SetVectorParameterValue(FName("ShadowColor"), AudioMaterialButtonStyle->ButtonShadowColor);
 			DynamicMaterial->SetVectorParameterValue(FName("SmoothBevelColor"), AudioMaterialButtonStyle->ButtonAccentColor);
 			DynamicMaterial->SetVectorParameterValue(FName("Color_1"), AudioMaterialButtonStyle->ButtonPressedMainColor);
-			DynamicMaterial->SetVectorParameterValue(FName("Color_2"), AudioMaterialButtonStyle->ButtonShadowMainColor);
+			DynamicMaterial->SetVectorParameterValue(FName("Color_2"), AudioMaterialButtonStyle->ButtonPressedShadowColor);
 			DynamicMaterial->SetVectorParameterValue(FName("LedColor"), AudioMaterialButtonStyle->ButtonPressedOutlineColor);
 			DynamicMaterial->SetScalarParameterValue(FName("Click"), bIsPressedAttribute.Get());
 
diff --git a/Engine/Plugins/Runtime/AudioWidgets/Source/AudioWidgets/Private/AudioMaterialSlate/SAudioMaterialKnob.cpp b/Engine/Plugins/Runtime/AudioWidgets/Source/AudioWidgets/Private/AudioMaterialSlate/SAudioMaterialKnob.cpp
--- a/Engine/Plugins/Runtime/AudioWidgets/Source/AudioWidgets/Private/AudioMaterialSlate/SAudioMaterialKnob.cpp
+++ b/Engine/Plugins/Runtime/AudioWidgets/Source/AudioWidgets/Private/AudioMaterialSlate/SAudioMaterialKnob.cpp
@@ -2,6 +2,7 @@
 
 
 #include "AudioMaterialSlate/SAudioMaterialKnob.h"
+#include "AudioMaterialSlate/AudioMaterialSlateTypes.h"
 #include "SlateOptMacros.h"
 #include "Components/Widget.h"
 
@@ -53,6 +54,7 @@ int32 SAudioMaterialKnob::OnPaint(const FPaintArgs& Args, const FGeometry& Allot
 			DynamicMaterial->SetVectorParameterValue(FName("LEDglowMin"), AudioMaterialKnobStyle->KnobBarFillMinColor);
 			DynamicMaterial->SetVectorParameterValue(FName("LED_Min"), AudioMaterialKnobStyle->KnobBarFillMinColor);
 			DynamicMaterial->SetVectorParameterValue(FName("LineColor"), AudioMaterialKnobStyle->KnobIndicatorColor);
+			DynamicMaterial->SetVectorParameterValue(FName("DotColor"), AudioMaterialKnobStyle->KnobIndicatorColor);
 			DynamicMaterial->SetVectorParameterValue(FName("LedTint"), AudioMaterialKnobStyle->KnobBarFillTintColor);
 
 			DynamicMaterial->SetScalarParameterValue(FName("VALUE"), FMath::Clamp(KnobPercent, 0.f, 1.f));
diff --git a/Engine/Plugins/Runtime/AudioWidgets/Source/AudioWidgets/Private/AudioWidgetsStyle.cpp b/Engine/Plugins/Runtime/AudioWidgets/Source/AudioWidgets/Private/AudioWidgetsStyle.cpp
--- a/Engine/Plugins/Runtime/AudioWidgets/Source/AudioWidgets/Private/AudioWidgetsStyle.cpp
+++ b/Engine/Plugins/Runtime/AudioWidgets/Source/AudioWidgets/Private/AudioWidgetsStyle.cpp
@@ -2,6 +2,7 @@
 
 #include "AudioWidgetsStyle.h"
 #include "AudioWidgetsSlateTypes.h"
+#include "AudioMaterialSlate/AudioMaterialSlateTypes.h"
 #include "Brushes/SlateNoResource.h"
 #include "Brushes/SlateRoundedBoxBrush.h"
 #include "Misc/Paths.h"
@@ -125,6 +126,45 @@ FAudioWidgetsStyle::FAudioWidgetsStyle()
 	*/
 	Set("ValueGridOverlay.Style", FSampledSequenceValueGridOverlayStyle());
 
+	/**
+	**AudioMaterialKnob Style
+	*/
+	UMaterialInterface* KnobDefaultMaterial = LoadObject<UMaterialInterface>(nullptr, TEXT("/AudioWidgets/AudioMaterialSlate/MI_AudioMaterialKnob.MI_AudioMaterialKnob"));
+
+	Set("AudioMaterialKnob.Style", FAudioMaterialKnobStyle()
+		.SetMaterial(KnobDefaultMaterial));
+
+	/**
+	**AudioMaterialMeter Style
+	*/
+	UMaterialInterface* MeterDefaultMaterial = LoadObject<UMaterialInterface>(nullptr, TEXT("/AudioWidgets/AudioMaterialSlate/MI_AudioMaterialMeter.MI_AudioMaterialMeter"));
+
+	Set("AudioMaterialMeter.Style", FAudioMaterialMeterStyle()
+		.SetMaterial(MeterDefaultMaterial));
+
+	/**
+	**AudioMaterialEnvelope Style
+	*/
+	UMaterialInterface* EnvelopeDefaultMaterial = LoadObject<UMaterialInterface>(nullptr, TEXT("/AudioWidgets/AudioMaterialSlate/MI_AudioMaterialEnvelope_ADSR.MI_AudioMaterialEnvelope_ADSR"));
+
+	Set("AudioMaterialEnvelope.Style", FAudioMaterialEnvelopeStyle()
+		.SetMaterial(EnvelopeDefaultMaterial));
+
+	/**
+	**AudioMaterialButton Style
+	*/
+	UMaterialInterface* ButtonDefaultMaterial = LoadObject<UMaterialInterface>(nullptr, TEXT("/AudioWidgets/AudioMaterialSlate/MI_AudioMaterialButton.MI_AudioMaterialButton"));
+
+	Set("AudioMaterialButton.Style", FAudioMaterialButtonStyle()
+		.SetMaterial(ButtonDefaultMaterial));
+
+	/**
+	**AudioMaterialSlider Style
+	*/
+	UMaterialInterface* SliderDefaultMaterial = LoadObject<UMaterialInterface>(nullptr, TEXT("/AudioWidgets/AudioMaterialSlate/MI_AudioMaterialSlider.MI_AudioMaterialSlider"));
+
+	Set("AudioMaterialSlider.Style", FAudioMaterialSliderStyle()
+		.SetMaterial(SliderDefaultMaterial));
 
 	FSlateStyleRegistry::RegisterSlateStyle(*this);
 }
diff --git a/Engine/Plugins/Runtime/AudioWidgets/Source/AudioWidgets/Public/AudioMaterialSlate/AudioMaterialButton.h b/Engine/Plugins/Runtime/AudioWidgets/Source/AudioWidgets/Public/AudioMaterialSlate/AudioMaterialButton.h
--- a/Engine/Plugins/Runtime/AudioWidgets/Source/AudioWidgets/Public/AudioMaterialSlate/AudioMaterialButton.h
+++ b/Engine/Plugins/Runtime/AudioWidgets/Source/AudioWidgets/Public/AudioMaterialSlate/AudioMaterialButton.h
@@ -2,7 +2,7 @@
 
 #pragma once
 
-#include "AudioMaterialSlateStyles.h"
+#include "AudioMaterialSlate/AudioMaterialSlateTypes.h"
 #include "Components/Widget.h"
 #include "Delegates/Delegate.h"
 #include "AudioMaterialButton.generated.h"
@@ -26,6 +26,8 @@ class AUDIOWIDGETS_API UAudioMaterialButton : public UWidget
 
 public:
 
+	UAudioMaterialButton();
+
 	/** The button's style */
 	UPROPERTY(EditAnywhere, BlueprintReadWrite, Category = "Style", meta = (DisplayName = "Style", ShowOnlyInnerProperties))
 	FAudioMaterialButtonStyle WidgetStyle;
diff --git a/Engine/Plugins/Runtime/AudioWidgets/Source/AudioWidgets/Public/AudioMaterialSlate/AudioMaterialEnvelope.h b/Engine/Plugins/Runtime/AudioWidgets/Source/AudioWidgets/Public/AudioMaterialSlate/AudioMaterialEnvelope.h
--- a/Engine/Plugins/Runtime/AudioWidgets/Source/AudioWidgets/Public/AudioMaterialSlate/AudioMaterialEnvelope.h
+++ b/Engine/Plugins/Runtime/AudioWidgets/Source/AudioWidgets/Public/AudioMaterialSlate/AudioMaterialEnvelope.h
@@ -3,7 +3,7 @@
 #pragma once
 
 #include "AudioMaterialEnvelopeSettings.h"
-#include "AudioMaterialSlateStyles.h"
+#include "AudioMaterialSlate/AudioMaterialSlateTypes.h"
 #include "Components/Widget.h"
 #include "Delegates/Delegate.h"
 #include "AudioMaterialEnvelope.generated.h"
@@ -23,6 +23,8 @@ class AUDIOWIDGETS_API UAudioMaterialEnvelope : public UWidget
 
 public:
 
+	UAudioMaterialEnvelope();
+
 	/** The Envelope's style */
 	UPROPERTY(EditAnywhere, BlueprintReadWrite, Category = "Style", meta = (DisplayName = "Style", ShowOnlyInnerProperties))
 	FAudioMaterialEnvelopeStyle WidgetStyle;
diff --git a/Engine/Plugins/Runtime/AudioWidgets/Source/AudioWidgets/Public/AudioMaterialSlate/AudioMaterialKnob.h b/Engine/Plugins/Runtime/AudioWidgets/Source/AudioWidgets/Public/AudioMaterialSlate/AudioMaterialKnob.h
--- a/Engine/Plugins/Runtime/AudioWidgets/Source/AudioWidgets/Public/AudioMaterialSlate/AudioMaterialKnob.h
+++ b/Engine/Plugins/Runtime/AudioWidgets/Source/AudioWidgets/Public/AudioMaterialSlate/AudioMaterialKnob.h
@@ -2,13 +2,12 @@
 
 #pragma once
 
-#include "AudioMaterialSlateStyles.h"
+#include "AudioMaterialSlate/AudioMaterialSlateTypes.h"
 #include "Components/Widget.h"
 #include "Delegates/Delegate.h"
 #include "AudioMaterialKnob.generated.h"
 
 class SAudioMaterialKnob;
-struct FAudioMaterialKnobStyle;
 
 DECLARE_DYNAMIC_MULTICAST_DELEGATE_OneParam(FOnKnobValueChangedEvent, float, Value);
 
@@ -26,6 +25,8 @@ class AUDIOWIDGETS_API UAudioMaterialKnob : public UWidget
 
 public:
 
+	UAudioMaterialKnob();
+
 	/** The button's style */
 	UPROPERTY(EditAnywhere, BlueprintReadWrite, Category = "Style", meta = (DisplayName = "Style", ShowOnlyInnerProperties))
 	FAudioMaterialKnobStyle WidgetStyle;
diff --git a/Engine/Plugins/Runtime/AudioWidgets/Source/AudioWidgets/Public/AudioMaterialSlate/AudioMaterialMeter.h b/Engine/Plugins/Runtime/AudioWidgets/Source/AudioWidgets/Public/AudioMaterialSlate/AudioMaterialMeter.h
--- a/Engine/Plugins/Runtime/AudioWidgets/Source/AudioWidgets/Public/AudioMaterialSlate/AudioMaterialMeter.h
+++ b/Engine/Plugins/Runtime/AudioWidgets/Source/AudioWidgets/Public/AudioMaterialSlate/AudioMaterialMeter.h
@@ -2,7 +2,7 @@
 
 #pragma once
 
-#include "AudioMaterialSlateStyles.h"
+#include "AudioMaterialSlate/AudioMaterialSlateTypes.h"
 #include "Components/Widget.h"
 #include "Delegates/Delegate.h"
 #include "AudioMaterialMeter.generated.h"
@@ -23,6 +23,8 @@ class AUDIOWIDGETS_API UAudioMaterialMeter : public UWidget
 
 public:
 
+	UAudioMaterialMeter();
+
 	/** The meter's style */
 	UPROPERTY(EditAnywhere, BlueprintReadWrite, Category = "Style", meta = (DisplayName = "Style", ShowOnlyInnerProperties))
 	FAudioMaterialMeterStyle WidgetStyle;
diff --git a/Engine/Plugins/Runtime/AudioWidgets/Source/AudioWidgets/Public/AudioMaterialSlate/AudioMaterialSlateStyles.h b/Engine/Plugins/Runtime/AudioWidgets/Source/AudioWidgets/Public/AudioMaterialSlate/AudioMaterialSlateStyles.h
deleted file mode 100644
diff --git a/Engine/Plugins/Runtime/AudioWidgets/Source/AudioWidgets/Public/AudioMaterialSlate/AudioMaterialSlateTypes.h b/Engine/Plugins/Runtime/AudioWidgets/Source/AudioWidgets/Public/AudioMaterialSlate/AudioMaterialSlateTypes.h
new file mode 100644
--- /dev/null
+++ b/Engine/Plugins/Runtime/AudioWidgets/Source/AudioWidgets/Public/AudioMaterialSlate/AudioMaterialSlateTypes.h
@@ -0,0 +1,249 @@
+// Copyright Epic Games, Inc. All Rights Reserved.
+
+#pragma once
+
+#include "Math/Color.h"
+#include "Materials/MaterialInstanceDynamic.h"
+#include "Materials/MaterialInterface.h"
+#include "Styling/ISlateStyle.h"
+#include "Styling/SlateWidgetStyle.h"
+#include "AudioMaterialSlateTypes.generated.h"
+
+/**
+ *Base for the appearance of an Audio Material Slates 
+ */
+USTRUCT(BlueprintType)
+struct AUDIOWIDGETS_API FAudioMaterialWidgetStyle : public FSlateWidgetStyle
+{
+	GENERATED_USTRUCT_BODY()
+
+public:
+
+	FAudioMaterialWidgetStyle();
+
+	/** Material used to render the Slate */
+	UPROPERTY(EditAnywhere, BlueprintReadWrite, meta = (DisplayPriority = 0), Category = "Style")
+	TObjectPtr<UMaterialInterface> Material = nullptr;
+
+	/** Desired Draw size of the rendered material*/
+	UPROPERTY(EditAnywhere, BlueprintReadWrite, meta = (DisplayPriority = 1), Category = "Style")
+	FVector2f DesiredSize;
+
+public:
+
+	void CreateDynamicMaterial(UObject* InOuter) const;
+	TObjectPtr<UMaterialInstanceDynamic> GetDynamicMaterial() const;
+
+private:
+
+	UPROPERTY()
+	mutable TObjectPtr<UMaterialInstanceDynamic> DynamicMaterial;
+};
+
+/**
+ *Represents the appearance of an Audio Material Button 
+ */
+USTRUCT(BlueprintType)
+struct AUDIOWIDGETS_API FAudioMaterialButtonStyle : public FAudioMaterialWidgetStyle
+{
+	GENERATED_USTRUCT_BODY()
+
+	FAudioMaterialButtonStyle();
+
+	// FSlateWidgetStyle
+	virtual void GetResources(TArray<const FSlateBrush*>& OutBrushes) const override;
+	static const FName TypeName;
+	virtual const FName GetTypeName() const override { return TypeName; };
+	static const FAudioMaterialButtonStyle& GetDefault();
+
+	FAudioMaterialButtonStyle& SetMaterial(UMaterialInterface* InMaterialInterface) { Material = InMaterialInterface; return *this; }
+
+	UPROPERTY(EditAnywhere, BlueprintReadWrite, Category = "Style")
+	FLinearColor ButtonMainColor;
+	FAudioMaterialButtonStyle& SetButtonMainColor(const FLinearColor& InColor) { ButtonMainColor = InColor; return *this; }
+
+	UPROPERTY(EditAnywhere, BlueprintReadWrite, Category = "Style")
+	FLinearColor ButtonShadowColor;
+	FAudioMaterialButtonStyle& SetButtonShadowColor(const FLinearColor& InColor) { ButtonShadowColor = InColor; return *this; }
+
+	UPROPERTY(EditAnywhere, BlueprintReadWrite, Category = "Style")
+	FLinearColor ButtonAccentColor;	
+	FAudioMaterialButtonStyle& SetButtonAccentColor(const FLinearColor& InColor) { ButtonShadowColor = InColor; return *this; }
+	
+	UPROPERTY(EditAnywhere, BlueprintReadWrite, Category = "Style")
+	FLinearColor ButtonPressedMainColor;	
+	FAudioMaterialButtonStyle& SetButtonPressedMainColor(const FLinearColor& InColor) { ButtonPressedMainColor = InColor; return *this; }
+
+	UPROPERTY(EditAnywhere, BlueprintReadWrite, Category = "Style")
+	FLinearColor ButtonPressedShadowColor;	
+	FAudioMaterialButtonStyle& SetButtonPressedShadowColor(const FLinearColor& InColor) { ButtonPressedShadowColor = InColor; return *this; }
+
+	UPROPERTY(EditAnywhere, BlueprintReadWrite, Category = "Style")
+	FLinearColor ButtonPressedOutlineColor;
+	FAudioMaterialButtonStyle& SetButtonPressedOutlineColor(const FLinearColor& InColor) { ButtonPressedOutlineColor = InColor; return *this; }
+
+};
+
+/**
+ *Represents the appearance of an Audio Material Slider 
+ */
+USTRUCT(BlueprintType)
+struct AUDIOWIDGETS_API FAudioMaterialSliderStyle : public FAudioMaterialWidgetStyle
+{
+	GENERATED_USTRUCT_BODY()
+
+	FAudioMaterialSliderStyle();
+
+	// FSlateWidgetStyle
+	virtual void GetResources(TArray<const FSlateBrush*>& OutBrushes) const override;
+	static const FName TypeName;
+	virtual const FName GetTypeName() const override { return TypeName; };
+	static const FAudioMaterialSliderStyle& GetDefault();
+
+	FAudioMaterialSliderStyle& SetMaterial(UMaterialInterface* InMaterialInterface) { Material = InMaterialInterface; return *this; }
+
+	UPROPERTY(EditAnywhere, BlueprintReadWrite, Category = "Style")
+	FLinearColor BarMainColor;
+	FAudioMaterialSliderStyle& SetSliderBarMainColor(const FLinearColor& InColor) { BarMainColor = InColor; return *this; }
+	
+	UPROPERTY(EditAnywhere, BlueprintReadWrite, Category = "Style")
+	FLinearColor BarShadowColor;	
+	FAudioMaterialSliderStyle& SetSliderBarShadowColor(const FLinearColor& InColor) { BarShadowColor = InColor; return *this; }
+	
+	UPROPERTY(EditAnywhere, BlueprintReadWrite, Category = "Style")
+	FLinearColor BarAccentColor;
+	FAudioMaterialSliderStyle& SetSliderBarAccentColor(const FLinearColor& InColor) { BarAccentColor = InColor; return *this; }
+
+	UPROPERTY(EditAnywhere, BlueprintReadWrite, Category = "Style")
+	FLinearColor HandleMainColor;	
+	FAudioMaterialSliderStyle& SetSliderHandleMainColor(const FLinearColor& InColor) { HandleMainColor = InColor; return *this; }
+	
+	UPROPERTY(EditAnywhere, BlueprintReadWrite, Category = "Style")
+	FLinearColor HandleOutlineColor;
+	FAudioMaterialSliderStyle& SetSliderHandleOutlineColor(const FLinearColor& InColor) { HandleOutlineColor = InColor; return *this; }
+
+};
+
+/**
+ *Represents the appearance of an Audio Material Knob 
+ */
+USTRUCT(BlueprintType)
+struct AUDIOWIDGETS_API FAudioMaterialKnobStyle : public FAudioMaterialWidgetStyle
+{
+	GENERATED_USTRUCT_BODY()
+
+	FAudioMaterialKnobStyle();
+
+	// FSlateWidgetStyle
+	virtual void GetResources(TArray<const FSlateBrush*>& OutBrushes) const override;
+	static const FName TypeName;
+	virtual const FName GetTypeName() const override { return TypeName; };
+	static const FAudioMaterialKnobStyle& GetDefault();
+
+	FAudioMaterialKnobStyle& SetMaterial(UMaterialInterface* InMaterialInterface) { Material = InMaterialInterface; return *this; }
+
+	UPROPERTY(EditAnywhere, BlueprintReadWrite, Category = "Style")
+	FLinearColor KnobMainColor;
+	FAudioMaterialKnobStyle& SetKnobMainColor(const FLinearColor& InColor) { KnobMainColor = InColor; return *this; }
+
+	UPROPERTY(EditAnywhere, BlueprintReadWrite, Category = "Style")
+	FLinearColor KnobAccentColor;
+	FAudioMaterialKnobStyle& SetKnobAccentColor(const FLinearColor& InColor) { KnobAccentColor = InColor; return *this; }
+	
+	UPROPERTY(EditAnywhere, BlueprintReadWrite, Category = "Style")
+	FLinearColor KnobIndicatorColor;
+	FAudioMaterialKnobStyle& SetKnobIndicatorColor(const FLinearColor& InColor) {KnobIndicatorColor = InColor; return *this;}
+	
+	UPROPERTY(EditAnywhere, BlueprintReadWrite, Category = "Style")
+	FLinearColor KnobBarColor;
+	FAudioMaterialKnobStyle& SetKnobBarColor(const FLinearColor& InColor) { KnobBarColor = InColor; return *this; }
+	
+	UPROPERTY(EditAnywhere, BlueprintReadWrite, Category = "Style")
+	FLinearColor KnobBarShadowColor;
+	FAudioMaterialKnobStyle& SetKnobBarShadowColor(const FLinearColor& InColor) { KnobBarShadowColor = InColor; return *this; }
+
+	UPROPERTY(EditAnywhere, BlueprintReadWrite, Category = "Style")
+	FLinearColor KnobBarFillMinColor;
+	FAudioMaterialKnobStyle& SetKnobBarFillMinColor(const FLinearColor& InColor) { KnobBarFillMinColor = InColor; return *this; }
+	
+	UPROPERTY(EditAnywhere, BlueprintReadWrite, Category = "Style")
+	FLinearColor KnobBarFillMidColor;
+	FAudioMaterialKnobStyle& SetKnobFillMidColor(const FLinearColor& InColor) { KnobBarFillMidColor = InColor; return *this; }
+	
+	UPROPERTY(EditAnywhere, BlueprintReadWrite, Category = "Style")
+	FLinearColor KnobBarFillMaxColor;
+	FAudioMaterialKnobStyle& SetKnobBarFillMaxColor(const FLinearColor& InColor) { KnobBarFillMaxColor = InColor; return *this; }
+	
+	UPROPERTY(EditAnywhere, BlueprintReadWrite, Category = "Style")
+	FLinearColor KnobBarFillTintColor;
+	FAudioMaterialKnobStyle& SetKnobBarFillTintColor(const FLinearColor& InColor) { KnobBarFillTintColor = InColor; return *this; }
+
+};
+
+/**
+ *Represents the appearance of an Audio Material Meter
+ */
+USTRUCT(BlueprintType)
+struct AUDIOWIDGETS_API FAudioMaterialMeterStyle : public FAudioMaterialWidgetStyle
+{
+	GENERATED_USTRUCT_BODY()
+
+	FAudioMaterialMeterStyle();
+
+	// FSlateWidgetStyle
+	virtual void GetResources(TArray<const FSlateBrush*>& OutBrushes) const override;
+	static const FName TypeName;
+	virtual const FName GetTypeName() const override { return TypeName; };
+	static const FAudioMaterialMeterStyle& GetDefault();
+
+	FAudioMaterialMeterStyle& SetMaterial(UMaterialInterface* InMaterialInterface) { Material = InMaterialInterface; return *this; }
+
+	UPROPERTY(EditAnywhere, BlueprintReadWrite, Category = "Style")
+	FLinearColor MeterFillMinColor;
+	FAudioMaterialMeterStyle& SetMeterFillMinColor(const FLinearColor& InColor) { MeterFillMinColor = InColor; return *this; }
+
+	UPROPERTY(EditAnywhere, BlueprintReadWrite, Category = "Style")
+	FLinearColor MeterFillMidColor;
+	FAudioMaterialMeterStyle& SetMeterFillMidColor(const FLinearColor& InColor) { MeterFillMidColor = InColor; return *this; }
+
+	UPROPERTY(EditAnywhere, BlueprintReadWrite, Category = "Style")
+	FLinearColor MeterFillMaxColor;
+	FAudioMaterialMeterStyle& SetMeterFillMaxColor(const FLinearColor& InColor) { MeterFillMaxColor = InColor; return *this; }
+
+	UPROPERTY(EditAnywhere, BlueprintReadWrite, Category = "Style")
+	FLinearColor MeterOffFillColor;
+	FAudioMaterialMeterStyle& SetMeterOffFillColor(const FLinearColor& InColor) { MeterOffFillColor = InColor; return *this; }
+
+};
+
+/**
+ *Represents the appearance of an Audio Material Envelope
+ */
+USTRUCT(BlueprintType)
+struct AUDIOWIDGETS_API FAudioMaterialEnvelopeStyle : public FAudioMaterialWidgetStyle
+{
+	GENERATED_USTRUCT_BODY()
+
+	FAudioMaterialEnvelopeStyle();
+
+	// FSlateWidgetStyle
+	virtual void GetResources(TArray<const FSlateBrush*>& OutBrushes) const override;
+	static const FName TypeName;
+	virtual const FName GetTypeName() const override { return TypeName; };
+	static const FAudioMaterialEnvelopeStyle& GetDefault();
+
+	FAudioMaterialEnvelopeStyle& SetMaterial(UMaterialInterface* InMaterialInterface) { Material = InMaterialInterface; return *this; }
+
+	UPROPERTY(EditAnywhere, BlueprintReadWrite, Category = "Style")
+	FLinearColor CurveColor;
+	FAudioMaterialEnvelopeStyle& SetEnvelopeCurveColor(const FLinearColor& InColor) { CurveColor = InColor; return *this; }
+	
+	UPROPERTY(EditAnywhere, BlueprintReadWrite, Category = "Style")
+	FLinearColor BackgroundColor;	
+	FAudioMaterialEnvelopeStyle& SetEnvelopeBackgroundColor(const FLinearColor& InColor) { BackgroundColor = InColor; return *this; }
+	
+	UPROPERTY(EditAnywhere, BlueprintReadWrite, Category = "Style")
+	FLinearColor OutlineColor;
+	FAudioMaterialEnvelopeStyle& SetEnvelopeOutlineColor(const FLinearColor& InColor) { OutlineColor = InColor; return *this; }
+	
+};
\ No newline at end of file
diff --git a/Engine/Plugins/Runtime/AudioWidgets/Source/AudioWidgets/Public/AudioMaterialSlate/AudioMaterialSlider.h b/Engine/Plugins/Runtime/AudioWidgets/Source/AudioWidgets/Public/AudioMaterialSlate/AudioMaterialSlider.h
--- a/Engine/Plugins/Runtime/AudioWidgets/Source/AudioWidgets/Public/AudioMaterialSlate/AudioMaterialSlider.h
+++ b/Engine/Plugins/Runtime/AudioWidgets/Source/AudioWidgets/Public/AudioMaterialSlate/AudioMaterialSlider.h
@@ -2,7 +2,7 @@
 
 #pragma once
 
-#include "AudioMaterialSlateStyles.h"
+#include "AudioMaterialSlate/AudioMaterialSlateTypes.h"
 #include "Components/Widget.h"
 #include "Delegates/Delegate.h"
 #include "AudioMaterialSlider.generated.h"
@@ -24,6 +24,8 @@ class AUDIOWIDGETS_API UAudioMaterialSlider : public UWidget
 
 public:
 
+	UAudioMaterialSlider();
+
 	/** The slider's style */
 	UPROPERTY(EditAnywhere, BlueprintReadWrite, Category = "Style", meta = (DisplayName = "Style", ShowOnlyInnerProperties))
 	FAudioMaterialSliderStyle WidgetStyle;
diff --git a/Engine/Plugins/Runtime/AudioWidgets/Source/AudioWidgets/Public/AudioMaterialSlate/AudioMaterialStyleContainers.h b/Engine/Plugins/Runtime/AudioWidgets/Source/AudioWidgets/Public/AudioMaterialSlate/AudioMaterialStyleContainers.h
new file mode 100644
--- /dev/null
+++ b/Engine/Plugins/Runtime/AudioWidgets/Source/AudioWidgets/Public/AudioMaterialSlate/AudioMaterialStyleContainers.h
@@ -0,0 +1,93 @@
+// Copyright Epic Games, Inc. All Rights Reserved.
+
+#pragma once
+
+#include "AudioMaterialSlateTypes.h"
+#include "Styling/SlateStyle.h"
+#include "Styling/SlateWidgetStyleContainerBase.h"
+#include "AudioMaterialStyleContainers.generated.h"
+
+UCLASS(hidecategories = Object, MinimalAPI)
+class UAudioMaterialKnobWidgetStyle : public USlateWidgetStyleContainerBase
+{
+	GENERATED_BODY()
+
+public:
+
+	/** The actual data describing the AudioMaterialKnob appearance. */
+	UPROPERTY(Category = Appearance, EditAnywhere, meta = (ShowOnlyInnerProperties))
+	FAudioMaterialKnobStyle KnobStyle;
+
+	virtual const struct FSlateWidgetStyle* const GetStyle() const override
+	{
+		return static_cast<const struct FSlateWidgetStyle*>(&KnobStyle);
+	}
+};
+
+UCLASS(hidecategories = Object, MinimalAPI)
+class UAudioMaterialMeterWidgetStyle : public USlateWidgetStyleContainerBase
+{
+	GENERATED_BODY()
+
+public:
+
+	/** The actual data describing the AudioMaterialMeter appearance. */
+	UPROPERTY(Category = Appearance, EditAnywhere, meta = (ShowOnlyInnerProperties))
+	FAudioMaterialMeterStyle MeterStyle;
+
+	virtual const struct FSlateWidgetStyle* const GetStyle() const override
+	{
+		return static_cast<const struct FSlateWidgetStyle*>(&MeterStyle);
+	}
+};
+
+UCLASS(hidecategories = Object, MinimalAPI)
+class UAudioMaterialButtonWidgetStyle : public USlateWidgetStyleContainerBase
+{
+	GENERATED_BODY()
+
+public:
+
+	/** The actual data describing the AudioMaterialButton appearance. */
+	UPROPERTY(Category = Appearance, EditAnywhere, meta = (ShowOnlyInnerProperties))
+	FAudioMaterialButtonStyle ButtonStyle;
+
+	virtual const struct FSlateWidgetStyle* const GetStyle() const override
+	{
+		return static_cast<const struct FSlateWidgetStyle*>(&ButtonStyle);
+	}
+};
+
+UCLASS(hidecategories = Object, MinimalAPI)
+class UAudioMaterialSliderWidgetStyle : public USlateWidgetStyleContainerBase
+{
+	GENERATED_BODY()
+
+public:
+
+	/** The actual data describing the AudioMaterialSlider appearance. */
+	UPROPERTY(Category = Appearance, EditAnywhere, meta = (ShowOnlyInnerProperties))
+	FAudioMaterialSliderStyle SliderStyle;
+
+	virtual const struct FSlateWidgetStyle* const GetStyle() const override
+	{
+		return static_cast<const struct FSlateWidgetStyle*>(&SliderStyle);
+	}
+};
+
+UCLASS(hidecategories = Object, MinimalAPI)
+class UAudioMaterialEnvelopeWidgetStyle : public USlateWidgetStyleContainerBase
+{
+	GENERATED_BODY()
+
+public:
+
+	/** The actual data describing the AudioMaterialSlider appearance. */
+	UPROPERTY(Category = Appearance, EditAnywhere, meta = (ShowOnlyInnerProperties))
+	FAudioMaterialEnvelopeStyle EnvelopeStyle;
+
+	virtual const struct FSlateWidgetStyle* const GetStyle() const override
+	{
+		return static_cast<const struct FSlateWidgetStyle*>(&EnvelopeStyle);
+	}
+};
diff --git a/Engine/Plugins/Runtime/AudioWidgets/Source/AudioWidgets/Public/AudioMaterialSlate/SAudioMaterialButton.h b/Engine/Plugins/Runtime/AudioWidgets/Source/AudioWidgets/Public/AudioMaterialSlate/SAudioMaterialButton.h
--- a/Engine/Plugins/Runtime/AudioWidgets/Source/AudioWidgets/Public/AudioMaterialSlate/SAudioMaterialButton.h
+++ b/Engine/Plugins/Runtime/AudioWidgets/Source/AudioWidgets/Public/AudioMaterialSlate/SAudioMaterialButton.h
@@ -2,7 +2,6 @@
 
 #pragma once
 
-#include "AudioMaterialSlate/AudioMaterialSlateStyles.h"
 #include "Framework/SlateDelegates.h"
 #include "Styling/ISlateStyle.h"
 #include "Styling/SlateWidgetStyleAsset.h"
diff --git a/Engine/Plugins/Runtime/AudioWidgets/Source/AudioWidgets/Public/AudioMaterialSlate/SAudioMaterialEnvelope.h b/Engine/Plugins/Runtime/AudioWidgets/Source/AudioWidgets/Public/AudioMaterialSlate/SAudioMaterialEnvelope.h
--- a/Engine/Plugins/Runtime/AudioWidgets/Source/AudioWidgets/Public/AudioMaterialSlate/SAudioMaterialEnvelope.h
+++ b/Engine/Plugins/Runtime/AudioWidgets/Source/AudioWidgets/Public/AudioMaterialSlate/SAudioMaterialEnvelope.h
@@ -2,7 +2,7 @@
 
 #pragma once
 
-#include "AudioMaterialSlate/AudioMaterialSlateStyles.h"
+#include "AudioMaterialSlate/AudioMaterialSlateTypes.h"
 #include "AudioMaterialEnvelopeSettings.h"
 #include "Framework/SlateDelegates.h"
 #include "Styling/ISlateStyle.h"
diff --git a/Engine/Plugins/Runtime/AudioWidgets/Source/AudioWidgets/Public/AudioMaterialSlate/SAudioMaterialKnob.h b/Engine/Plugins/Runtime/AudioWidgets/Source/AudioWidgets/Public/AudioMaterialSlate/SAudioMaterialKnob.h
--- a/Engine/Plugins/Runtime/AudioWidgets/Source/AudioWidgets/Public/AudioMaterialSlate/SAudioMaterialKnob.h
+++ b/Engine/Plugins/Runtime/AudioWidgets/Source/AudioWidgets/Public/AudioMaterialSlate/SAudioMaterialKnob.h
@@ -2,7 +2,7 @@
 
 #pragma once
 
-#include "AudioMaterialSlate/AudioMaterialSlateStyles.h"
+#include "AudioMaterialSlate/AudioMaterialSlateTypes.h"
 #include "Framework/SlateDelegates.h"
 #include "Styling/ISlateStyle.h"
 #include "Styling/SlateWidgetStyleAsset.h"
diff --git a/Engine/Plugins/Runtime/AudioWidgets/Source/AudioWidgets/Public/AudioMaterialSlate/SAudioMaterialMeter.h b/Engine/Plugins/Runtime/AudioWidgets/Source/AudioWidgets/Public/AudioMaterialSlate/SAudioMaterialMeter.h
--- a/Engine/Plugins/Runtime/AudioWidgets/Source/AudioWidgets/Public/AudioMaterialSlate/SAudioMaterialMeter.h
+++ b/Engine/Plugins/Runtime/AudioWidgets/Source/AudioWidgets/Public/AudioMaterialSlate/SAudioMaterialMeter.h
@@ -2,7 +2,7 @@
 
 #pragma once
 
-#include "AudioMaterialSlate/AudioMaterialSlateStyles.h"
+#include "AudioMaterialSlate/AudioMaterialSlateTypes.h"
 #include "Framework/SlateDelegates.h"
 #include "Styling/ISlateStyle.h"
 #include "Styling/SlateWidgetStyleAsset.h"
diff --git a/Engine/Plugins/Runtime/AudioWidgets/Source/AudioWidgets/Public/AudioMaterialSlate/SAudioMaterialSlider.h b/Engine/Plugins/Runtime/AudioWidgets/Source/AudioWidgets/Public/AudioMaterialSlate/SAudioMaterialSlider.h
--- a/Engine/Plugins/Runtime/AudioWidgets/Source/AudioWidgets/Public/AudioMaterialSlate/SAudioMaterialSlider.h
+++ b/Engine/Plugins/Runtime/AudioWidgets/Source/AudioWidgets/Public/AudioMaterialSlate/SAudioMaterialSlider.h
@@ -2,7 +2,7 @@
 
 #pragma once
 
-#include "AudioMaterialSlate/AudioMaterialSlateStyles.h"
+#include "AudioMaterialSlate/AudioMaterialSlateTypes.h"
 #include "Framework/SlateDelegates.h"
 #include "Styling/ISlateStyle.h"
 #include "Styling/SlateWidgetStyleAsset.h"
