[Backout] - CL31556103 - CIS Warnings  430673
Original CL Desc
-----------------------------------------------------------------
DefaultStyleSets for AudioMaterialSlates

#virtualized

diff --git a/Engine/Plugins/Runtime/AudioWidgets/Source/AudioWidgets/Private/AudioMaterialSlate/AudioMaterialButton.cpp b/Engine/Plugins/Runtime/AudioWidgets/Source/AudioWidgets/Private/AudioMaterialSlate/AudioMaterialButton.cpp
--- a/Engine/Plugins/Runtime/AudioWidgets/Source/AudioWidgets/Private/AudioMaterialSlate/AudioMaterialButton.cpp
+++ b/Engine/Plugins/Runtime/AudioWidgets/Source/AudioWidgets/Private/AudioMaterialSlate/AudioMaterialButton.cpp
@@ -2,18 +2,9 @@
 
 
 #include "AudioMaterialSlate/AudioMaterialButton.h"
-#include "AudioMaterialSlate/AudioMaterialSlateTypes.h"
 #include "AudioMaterialSlate/SAudioMaterialButton.h"
-#include "AudioWidgetsStyle.h"
 
 #define LOCTEXT_NAMESPACE "AudioWidgets"
-
-UAudioMaterialButton::UAudioMaterialButton()
-{
-	//get default style
-	WidgetStyle = FAudioWidgetsStyle::Get().GetWidgetStyle<FAudioMaterialButtonStyle>("AudioMaterialButton.Style");
-}
-
 #if WITH_EDITOR
 const FText UAudioMaterialButton::GetPaletteCategory()
 {
diff --git a/Engine/Plugins/Runtime/AudioWidgets/Source/AudioWidgets/Private/AudioMaterialSlate/AudioMaterialEnvelope.cpp b/Engine/Plugins/Runtime/AudioWidgets/Source/AudioWidgets/Private/AudioMaterialSlate/AudioMaterialEnvelope.cpp
--- a/Engine/Plugins/Runtime/AudioWidgets/Source/AudioWidgets/Private/AudioMaterialSlate/AudioMaterialEnvelope.cpp
+++ b/Engine/Plugins/Runtime/AudioWidgets/Source/AudioWidgets/Private/AudioMaterialSlate/AudioMaterialEnvelope.cpp
@@ -2,20 +2,13 @@
 
 
 #include "AudioMaterialSlate/AudioMaterialEnvelope.h"
-#include "AudioMaterialSlate/AudioMaterialSlateTypes.h"
+#include "AudioMaterialSlate/AudioMaterialSlateStyles.h"
 #include "AudioMaterialSlate/SAudioMaterialEnvelope.h"
-#include "AudioWidgetsStyle.h"
 #include "Widgets/SWeakWidget.h"
 #include "Components/AudioComponent.h"
 
 #define LOCTEXT_NAMESPACE "AudioWidgets"
 
-UAudioMaterialEnvelope::UAudioMaterialEnvelope()
-{
-	//get default style
-	WidgetStyle = FAudioWidgetsStyle::Get().GetWidgetStyle<FAudioMaterialEnvelopeStyle>("AudioMaterialEnvelope.Style");
-}
-
 #if WITH_EDITOR
 const FText UAudioMaterialEnvelope::GetPaletteCategory()
 {
diff --git a/Engine/Plugins/Runtime/AudioWidgets/Source/AudioWidgets/Private/AudioMaterialSlate/AudioMaterialKnob.cpp b/Engine/Plugins/Runtime/AudioWidgets/Source/AudioWidgets/Private/AudioMaterialSlate/AudioMaterialKnob.cpp
--- a/Engine/Plugins/Runtime/AudioWidgets/Source/AudioWidgets/Private/AudioMaterialSlate/AudioMaterialKnob.cpp
+++ b/Engine/Plugins/Runtime/AudioWidgets/Source/AudioWidgets/Private/AudioMaterialSlate/AudioMaterialKnob.cpp
@@ -2,17 +2,9 @@
 
 
 #include "AudioMaterialSlate/AudioMaterialKnob.h"
-#include "AudioMaterialSlate/AudioMaterialSlateTypes.h"
 #include "AudioMaterialSlate/SAudioMaterialKnob.h"
-#include "AudioWidgetsStyle.h"
 
 #define LOCTEXT_NAMESPACE "AudioWidgets"
-UAudioMaterialKnob::UAudioMaterialKnob()
-{
-	//get default style
-	WidgetStyle = FAudioWidgetsStyle::Get().GetWidgetStyle<FAudioMaterialKnobStyle>("AudioMaterialKnob.Style");	
-}
-
 #if WITH_EDITOR
 const FText UAudioMaterialKnob::GetPaletteCategory()
 {
diff --git a/Engine/Plugins/Runtime/AudioWidgets/Source/AudioWidgets/Private/AudioMaterialSlate/AudioMaterialMeter.cpp b/Engine/Plugins/Runtime/AudioWidgets/Source/AudioWidgets/Private/AudioMaterialSlate/AudioMaterialMeter.cpp
--- a/Engine/Plugins/Runtime/AudioWidgets/Source/AudioWidgets/Private/AudioMaterialSlate/AudioMaterialMeter.cpp
+++ b/Engine/Plugins/Runtime/AudioWidgets/Source/AudioWidgets/Private/AudioMaterialSlate/AudioMaterialMeter.cpp
@@ -2,20 +2,13 @@
 
 
 #include "AudioMaterialSlate/AudioMaterialMeter.h"
-#include "AudioMaterialSlate/AudioMaterialSlateTypes.h"
+#include "AudioMaterialSlate/AudioMaterialSlateStyles.h"
 #include "AudioMaterialSlate/SAudioMaterialMeter.h"
-#include "AudioWidgetsStyle.h"
 #include "Widgets/SWeakWidget.h"
 #include "Components/AudioComponent.h"
 
 #define LOCTEXT_NAMESPACE "AudioWidgets"
 
-UAudioMaterialMeter::UAudioMaterialMeter()
-{
-	//get default style
-	WidgetStyle = FAudioWidgetsStyle::Get().GetWidgetStyle<FAudioMaterialMeterStyle>("AudioMaterialMeter.Style");
-}
-
 #if WITH_EDITOR
 const FText UAudioMaterialMeter::GetPaletteCategory()
 {
diff --git a/Engine/Plugins/Runtime/AudioWidgets/Source/AudioWidgets/Private/AudioMaterialSlate/AudioMaterialSlateStyles.cpp b/Engine/Plugins/Runtime/AudioWidgets/Source/AudioWidgets/Private/AudioMaterialSlate/AudioMaterialSlateStyles.cpp
new file mode 100644
--- /dev/null
+++ b/Engine/Plugins/Runtime/AudioWidgets/Source/AudioWidgets/Private/AudioMaterialSlate/AudioMaterialSlateStyles.cpp
@@ -0,0 +1,156 @@
+// Copyright Epic Games, Inc. All Rights Reserved.
+
+
+#include "AudioMaterialSlate/AudioMaterialSlateStyles.h"
+
+FAudioMaterialStyle::FAudioMaterialStyle()
+	:DesiredSize(32.f, 32.f)
+{
+}
+
+TObjectPtr<UMaterialInstanceDynamic> FAudioMaterialStyle::GetDynamicMaterial() const
+{
+	return DynamicMaterial;
+}
+
+void FAudioMaterialStyle::CreateDynamicMaterial(UObject* InOuter) const
+{
+	DynamicMaterial = UMaterialInstanceDynamic::Create(Material, InOuter);
+}
+
+void FAudioMaterialSliderStyle::GetResources(TArray<const FSlateBrush*>& OutBrushes) const
+{
+	// Add any brush resources here so that Slate can correctly atlas and reference them
+}
+
+FAudioMaterialButtonStyle::FAudioMaterialButtonStyle()
+	:ButtonMainColor(FLinearColor(0.1f, 0.1f, 0.1f, 1.0f))
+	,ButtonShadowColor(FLinearColor(0.0f, 0.0f, 0.0f, 1.0f))
+	,ButtonAccentColor(FLinearColor(0.34f, 0.34f, 0.34f, 1.0f))
+	,ButtonPressedMainColor(FLinearColor(0.95f, 0.95f, 0.95f, 1.0f))
+	,ButtonShadowMainColor(FLinearColor(0.156f, 0.138f, 0.15f, 1.0f))
+	,ButtonPressedOutlineColor(FLinearColor(0.156f, 0.138f, 0.15f, 1.0f))
+{
+	DesiredSize = FVector2f(128.f, 128.f);
+}
+
+FAudioMaterialButtonStyle::~FAudioMaterialButtonStyle()
+{
+}
+
+const FName FAudioMaterialButtonStyle::TypeName(TEXT("FAudioMaterialButtonStyle"));
+
+void FAudioMaterialButtonStyle::GetResources(TArray<const FSlateBrush*>& OutBrushes) const
+{
+}
+
+const FAudioMaterialButtonStyle& FAudioMaterialButtonStyle::GetDefault()
+{
+	static FAudioMaterialButtonStyle Default;
+	return Default;
+}
+
+FAudioMaterialSliderStyle::FAudioMaterialSliderStyle()
+	:BarMainColor(FLinearColor(1.0f, 1.0f, 1.0f, 0.0f))
+	,BarShadowColor(FLinearColor(1.0f, 1.0f, 1.0f, 0.0f))
+	,BarAccentColor(FLinearColor(1.0f, 1.0f, 1.0f, 0.0f))
+	,HandleMainColor(FLinearColor(1.0f, 1.0f, 1.0f, 0.0f))
+	,HandleOutlineColor(FLinearColor(1.0f, 1.0f, 1.0f, 0.0f))
+{
+	DesiredSize = FVector2f(128.f, 512.f);
+}
+
+FAudioMaterialSliderStyle::~FAudioMaterialSliderStyle()
+{
+}
+
+const FName FAudioMaterialSliderStyle::TypeName(TEXT("FAudioMaterialSliderStyle"));
+
+const FAudioMaterialSliderStyle& FAudioMaterialSliderStyle::GetDefault()
+{
+	static FAudioMaterialSliderStyle Default;
+	return Default;
+}
+
+FAudioMaterialKnobStyle::FAudioMaterialKnobStyle()
+	:KnobMainColor(FLinearColor(0.f, 0.f, 0.f, 1.0f))
+	,KnobAccentColor(FLinearColor(.14f, .16f, .2f, 1.0f))
+	,KnobIndicatorColor(FLinearColor(1.f, 0.f, 0.04f, 1.0f))
+	,KnobBarColor(FLinearColor(.033f, .033f, .033f, 1.0f))
+	,KnobBarShadowColor(FLinearColor(0.f, 0.f, 0.f, 1.0f))
+	,KnobBarFillMinColor(FLinearColor(0.9f, 0.9f, 0.9f, 1.0f))
+	,KnobBarFillMidColor(FLinearColor(0.9f, 0.9f, 0.9f, 1.0f))
+	,KnobBarFillMaxColor(FLinearColor(0.9f, 0.9f, 0.9f, 1.0f))
+	,KnobBarFillTintColor(FLinearColor(1.f, 1.f, 1.f, 1.0f))
+{
+	DesiredSize = FVector2f(256.f,256.f);
+}
+
+FAudioMaterialKnobStyle::~FAudioMaterialKnobStyle()
+{
+}
+
+const FName FAudioMaterialKnobStyle::TypeName(TEXT("FAudioMaterialKnobStyle"));
+
+const FAudioMaterialKnobStyle& FAudioMaterialKnobStyle::GetDefault()
+{
+	static FAudioMaterialKnobStyle Default;
+	return Default;
+}
+
+void FAudioMaterialKnobStyle::GetResources(TArray<const FSlateBrush*>& OutBrushes) const
+{
+	// Add any brush resources here so that Slate can correctly atlas and reference them
+}
+
+FAudioMaterialMeterStyle::FAudioMaterialMeterStyle()
+	:MeterFillMinColor(FLinearColor(1.f, 1.f, 1.f, 1.f))
+	,MeterFillMidColor(FLinearColor(1.f, 1.f, 1.f, 1.f))
+	,MeterFillMaxColor(FLinearColor(1.f, 1.f, 1.f, 1.f))
+	,MeterOffFillColor(FLinearColor(0.5f, 0.5f, 0.5f, 1.f))
+{
+	DesiredSize = FVector2f(32.f, 512.f);
+}
+
+FAudioMaterialMeterStyle::~FAudioMaterialMeterStyle()
+{
+}
+
+const FName FAudioMaterialMeterStyle::TypeName(TEXT("FAudioMaterialKnobStyle"));
+
+const FAudioMaterialMeterStyle& FAudioMaterialMeterStyle::GetDefault()
+{
+	static FAudioMaterialMeterStyle Default;
+	return Default;
+}
+
+void FAudioMaterialMeterStyle::GetResources(TArray<const FSlateBrush*>& OutBrushes) const
+{
+	// Add any brush resources here so that Slate can correctly atlas and reference them
+}
+
+
+FAudioMaterialEnvelopeStyle::FAudioMaterialEnvelopeStyle()
+	:CurveColor(FLinearColor(0.15f, 0.45f, 0.75f, 1.f))
+	,BackgroundColor(FLinearColor(0.01f, 0.01f, 0.01f, 1.f))
+	,OutlineColor(FLinearColor(0.05f, 0.05f, 0.05f, 1.f))
+{
+	DesiredSize = FVector2f(256.f, 256.f);
+}
+
+FAudioMaterialEnvelopeStyle::~FAudioMaterialEnvelopeStyle()
+{
+}
+
+const FName FAudioMaterialEnvelopeStyle::TypeName(TEXT("FAudioMaterialKnobStyle"));
+
+const FAudioMaterialEnvelopeStyle& FAudioMaterialEnvelopeStyle::GetDefault()
+{
+	static FAudioMaterialEnvelopeStyle Default;
+	return Default;
+}
+
+void FAudioMaterialEnvelopeStyle::GetResources(TArray<const FSlateBrush*>& OutBrushes) const
+{
+	// Add any brush resources here so that Slate can correctly atlas and reference them
+}
diff --git a/Engine/Plugins/Runtime/AudioWidgets/Source/AudioWidgets/Private/AudioMaterialSlate/AudioMaterialSlateTypes.cpp b/Engine/Plugins/Runtime/AudioWidgets/Source/AudioWidgets/Private/AudioMaterialSlate/AudioMaterialSlateTypes.cpp
deleted file mode 100644
diff --git a/Engine/Plugins/Runtime/AudioWidgets/Source/AudioWidgets/Private/AudioMaterialSlate/AudioMaterialSlider.cpp b/Engine/Plugins/Runtime/AudioWidgets/Source/AudioWidgets/Private/AudioMaterialSlate/AudioMaterialSlider.cpp
--- a/Engine/Plugins/Runtime/AudioWidgets/Source/AudioWidgets/Private/AudioMaterialSlate/AudioMaterialSlider.cpp
+++ b/Engine/Plugins/Runtime/AudioWidgets/Source/AudioWidgets/Private/AudioMaterialSlate/AudioMaterialSlider.cpp
@@ -2,19 +2,12 @@
 
 
 #include "AudioMaterialSlate/AudioMaterialSlider.h"
-#include "AudioMaterialSlate/AudioMaterialSlateTypes.h"
+#include "AudioMaterialSlate/AudioMaterialSlateStyles.h"
 #include "AudioMaterialSlate/SAudioMaterialSlider.h"
-#include "AudioWidgetsStyle.h"
 #include "Widgets/SWeakWidget.h"
 
 #define LOCTEXT_NAMESPACE "AudioWidgets"
 
-UAudioMaterialSlider::UAudioMaterialSlider()
-{
-	//get default style
-	WidgetStyle = FAudioWidgetsStyle::Get().GetWidgetStyle<FAudioMaterialSliderStyle>("AudioMaterialSlider.Style");
-}
-
 #if WITH_EDITOR
 const FText UAudioMaterialSlider::GetPaletteCategory()
 {
diff --git a/Engine/Plugins/Runtime/AudioWidgets/Source/AudioWidgets/Private/AudioMaterialSlate/SAudioMaterialButton.cpp b/Engine/Plugins/Runtime/AudioWidgets/Source/AudioWidgets/Private/AudioMaterialSlate/SAudioMaterialButton.cpp
--- a/Engine/Plugins/Runtime/AudioWidgets/Source/AudioWidgets/Private/AudioMaterialSlate/SAudioMaterialButton.cpp
+++ b/Engine/Plugins/Runtime/AudioWidgets/Source/AudioWidgets/Private/AudioMaterialSlate/SAudioMaterialButton.cpp
@@ -2,7 +2,6 @@
 
 
 #include "AudioMaterialSlate/SAudioMaterialButton.h"
-#include "AudioMaterialSlate/AudioMaterialSlateTypes.h"
 #include "SlateOptMacros.h"
 #include "Components/Widget.h"
 
@@ -44,7 +43,7 @@ int32 SAudioMaterialButton::OnPaint(const FPaintArgs& Args, const FGeometry& All
 			DynamicMaterial->SetVectorParameterValue(FName("ShadowColor"), AudioMaterialButtonStyle->ButtonShadowColor);
 			DynamicMaterial->SetVectorParameterValue(FName("SmoothBevelColor"), AudioMaterialButtonStyle->ButtonAccentColor);
 			DynamicMaterial->SetVectorParameterValue(FName("Color_1"), AudioMaterialButtonStyle->ButtonPressedMainColor);
-			DynamicMaterial->SetVectorParameterValue(FName("Color_2"), AudioMaterialButtonStyle->ButtonPressedShadowColor);
+			DynamicMaterial->SetVectorParameterValue(FName("Color_2"), AudioMaterialButtonStyle->ButtonShadowMainColor);
 			DynamicMaterial->SetVectorParameterValue(FName("LedColor"), AudioMaterialButtonStyle->ButtonPressedOutlineColor);
 			DynamicMaterial->SetScalarParameterValue(FName("Click"), bIsPressedAttribute.Get());
 
diff --git a/Engine/Plugins/Runtime/AudioWidgets/Source/AudioWidgets/Private/AudioMaterialSlate/SAudioMaterialKnob.cpp b/Engine/Plugins/Runtime/AudioWidgets/Source/AudioWidgets/Private/AudioMaterialSlate/SAudioMaterialKnob.cpp
--- a/Engine/Plugins/Runtime/AudioWidgets/Source/AudioWidgets/Private/AudioMaterialSlate/SAudioMaterialKnob.cpp
+++ b/Engine/Plugins/Runtime/AudioWidgets/Source/AudioWidgets/Private/AudioMaterialSlate/SAudioMaterialKnob.cpp
@@ -2,7 +2,6 @@
 
 
 #include "AudioMaterialSlate/SAudioMaterialKnob.h"
-#include "AudioMaterialSlate/AudioMaterialSlateTypes.h"
 #include "SlateOptMacros.h"
 #include "Components/Widget.h"
 
@@ -54,7 +53,6 @@ int32 SAudioMaterialKnob::OnPaint(const FPaintArgs& Args, const FGeometry& Allot
 			DynamicMaterial->SetVectorParameterValue(FName("LEDglowMin"), AudioMaterialKnobStyle->KnobBarFillMinColor);
 			DynamicMaterial->SetVectorParameterValue(FName("LED_Min"), AudioMaterialKnobStyle->KnobBarFillMinColor);
 			DynamicMaterial->SetVectorParameterValue(FName("LineColor"), AudioMaterialKnobStyle->KnobIndicatorColor);
-			DynamicMaterial->SetVectorParameterValue(FName("DotColor"), AudioMaterialKnobStyle->KnobIndicatorColor);
 			DynamicMaterial->SetVectorParameterValue(FName("LedTint"), AudioMaterialKnobStyle->KnobBarFillTintColor);
 
 			DynamicMaterial->SetScalarParameterValue(FName("VALUE"), FMath::Clamp(KnobPercent, 0.f, 1.f));
diff --git a/Engine/Plugins/Runtime/AudioWidgets/Source/AudioWidgets/Private/AudioWidgetsStyle.cpp b/Engine/Plugins/Runtime/AudioWidgets/Source/AudioWidgets/Private/AudioWidgetsStyle.cpp
--- a/Engine/Plugins/Runtime/AudioWidgets/Source/AudioWidgets/Private/AudioWidgetsStyle.cpp
+++ b/Engine/Plugins/Runtime/AudioWidgets/Source/AudioWidgets/Private/AudioWidgetsStyle.cpp
@@ -2,7 +2,6 @@
 
 #include "AudioWidgetsStyle.h"
 #include "AudioWidgetsSlateTypes.h"
-#include "AudioMaterialSlate/AudioMaterialSlateTypes.h"
 #include "Brushes/SlateNoResource.h"
 #include "Brushes/SlateRoundedBoxBrush.h"
 #include "Misc/Paths.h"
@@ -126,45 +125,6 @@ FAudioWidgetsStyle::FAudioWidgetsStyle()
 	*/
 	Set("ValueGridOverlay.Style", FSampledSequenceValueGridOverlayStyle());
 
-	/**
-	**AudioMaterialKnob Style
-	*/
-	UMaterialInterface* KnobDefaultMaterial = LoadObject<UMaterialInterface>(nullptr, TEXT("/AudioWidgets/AudioMaterialSlate/MI_AudioMaterialKnob.MI_AudioMaterialKnob"));
-
-	Set("AudioMaterialKnob.Style", FAudioMaterialKnobStyle()
-		.SetMaterial(KnobDefaultMaterial));
-
-	/**
-	**AudioMaterialMeter Style
-	*/
-	UMaterialInterface* MeterDefaultMaterial = LoadObject<UMaterialInterface>(nullptr, TEXT("/AudioWidgets/AudioMaterialSlate/MI_AudioMaterialMeter.MI_AudioMaterialMeter"));
-
-	Set("AudioMaterialMeter.Style", FAudioMaterialMeterStyle()
-		.SetMaterial(MeterDefaultMaterial));
-
-	/**
-	**AudioMaterialEnvelope Style
-	*/
-	UMaterialInterface* EnvelopeDefaultMaterial = LoadObject<UMaterialInterface>(nullptr, TEXT("/AudioWidgets/AudioMaterialSlate/MI_AudioMaterialEnvelope_ADSR.MI_AudioMaterialEnvelope_ADSR"));
-
-	Set("AudioMaterialEnvelope.Style", FAudioMaterialEnvelopeStyle()
-		.SetMaterial(EnvelopeDefaultMaterial));
-
-	/**
-	**AudioMaterialButton Style
-	*/
-	UMaterialInterface* ButtonDefaultMaterial = LoadObject<UMaterialInterface>(nullptr, TEXT("/AudioWidgets/AudioMaterialSlate/MI_AudioMaterialButton.MI_AudioMaterialButton"));
-
-	Set("AudioMaterialButton.Style", FAudioMaterialButtonStyle()
-		.SetMaterial(ButtonDefaultMaterial));
-
-	/**
-	**AudioMaterialSlider Style
-	*/
-	UMaterialInterface* SliderDefaultMaterial = LoadObject<UMaterialInterface>(nullptr, TEXT("/AudioWidgets/AudioMaterialSlate/MI_AudioMaterialSlider.MI_AudioMaterialSlider"));
-
-	Set("AudioMaterialSlider.Style", FAudioMaterialSliderStyle()
-		.SetMaterial(SliderDefaultMaterial));
 
 	FSlateStyleRegistry::RegisterSlateStyle(*this);
 }
diff --git a/Engine/Plugins/Runtime/AudioWidgets/Source/AudioWidgets/Public/AudioMaterialSlate/AudioMaterialButton.h b/Engine/Plugins/Runtime/AudioWidgets/Source/AudioWidgets/Public/AudioMaterialSlate/AudioMaterialButton.h
--- a/Engine/Plugins/Runtime/AudioWidgets/Source/AudioWidgets/Public/AudioMaterialSlate/AudioMaterialButton.h
+++ b/Engine/Plugins/Runtime/AudioWidgets/Source/AudioWidgets/Public/AudioMaterialSlate/AudioMaterialButton.h
@@ -2,7 +2,7 @@
 
 #pragma once
 
-#include "AudioMaterialSlate/AudioMaterialSlateTypes.h"
+#include "AudioMaterialSlateStyles.h"
 #include "Components/Widget.h"
 #include "Delegates/Delegate.h"
 #include "AudioMaterialButton.generated.h"
@@ -26,8 +26,6 @@ class AUDIOWIDGETS_API UAudioMaterialButton : public UWidget
 
 public:
 
-	UAudioMaterialButton();
-
 	/** The button's style */
 	UPROPERTY(EditAnywhere, BlueprintReadWrite, Category = "Style", meta = (DisplayName = "Style", ShowOnlyInnerProperties))
 	FAudioMaterialButtonStyle WidgetStyle;
diff --git a/Engine/Plugins/Runtime/AudioWidgets/Source/AudioWidgets/Public/AudioMaterialSlate/AudioMaterialEnvelope.h b/Engine/Plugins/Runtime/AudioWidgets/Source/AudioWidgets/Public/AudioMaterialSlate/AudioMaterialEnvelope.h
--- a/Engine/Plugins/Runtime/AudioWidgets/Source/AudioWidgets/Public/AudioMaterialSlate/AudioMaterialEnvelope.h
+++ b/Engine/Plugins/Runtime/AudioWidgets/Source/AudioWidgets/Public/AudioMaterialSlate/AudioMaterialEnvelope.h
@@ -3,7 +3,7 @@
 #pragma once
 
 #include "AudioMaterialEnvelopeSettings.h"
-#include "AudioMaterialSlate/AudioMaterialSlateTypes.h"
+#include "AudioMaterialSlateStyles.h"
 #include "Components/Widget.h"
 #include "Delegates/Delegate.h"
 #include "AudioMaterialEnvelope.generated.h"
@@ -23,8 +23,6 @@ class AUDIOWIDGETS_API UAudioMaterialEnvelope : public UWidget
 
 public:
 
-	UAudioMaterialEnvelope();
-
 	/** The Envelope's style */
 	UPROPERTY(EditAnywhere, BlueprintReadWrite, Category = "Style", meta = (DisplayName = "Style", ShowOnlyInnerProperties))
 	FAudioMaterialEnvelopeStyle WidgetStyle;
diff --git a/Engine/Plugins/Runtime/AudioWidgets/Source/AudioWidgets/Public/AudioMaterialSlate/AudioMaterialKnob.h b/Engine/Plugins/Runtime/AudioWidgets/Source/AudioWidgets/Public/AudioMaterialSlate/AudioMaterialKnob.h
--- a/Engine/Plugins/Runtime/AudioWidgets/Source/AudioWidgets/Public/AudioMaterialSlate/AudioMaterialKnob.h
+++ b/Engine/Plugins/Runtime/AudioWidgets/Source/AudioWidgets/Public/AudioMaterialSlate/AudioMaterialKnob.h
@@ -2,12 +2,13 @@
 
 #pragma once
 
-#include "AudioMaterialSlate/AudioMaterialSlateTypes.h"
+#include "AudioMaterialSlateStyles.h"
 #include "Components/Widget.h"
 #include "Delegates/Delegate.h"
 #include "AudioMaterialKnob.generated.h"
 
 class SAudioMaterialKnob;
+struct FAudioMaterialKnobStyle;
 
 DECLARE_DYNAMIC_MULTICAST_DELEGATE_OneParam(FOnKnobValueChangedEvent, float, Value);
 
@@ -25,8 +26,6 @@ class AUDIOWIDGETS_API UAudioMaterialKnob : public UWidget
 
 public:
 
-	UAudioMaterialKnob();
-
 	/** The button's style */
 	UPROPERTY(EditAnywhere, BlueprintReadWrite, Category = "Style", meta = (DisplayName = "Style", ShowOnlyInnerProperties))
 	FAudioMaterialKnobStyle WidgetStyle;
diff --git a/Engine/Plugins/Runtime/AudioWidgets/Source/AudioWidgets/Public/AudioMaterialSlate/AudioMaterialMeter.h b/Engine/Plugins/Runtime/AudioWidgets/Source/AudioWidgets/Public/AudioMaterialSlate/AudioMaterialMeter.h
--- a/Engine/Plugins/Runtime/AudioWidgets/Source/AudioWidgets/Public/AudioMaterialSlate/AudioMaterialMeter.h
+++ b/Engine/Plugins/Runtime/AudioWidgets/Source/AudioWidgets/Public/AudioMaterialSlate/AudioMaterialMeter.h
@@ -2,7 +2,7 @@
 
 #pragma once
 
-#include "AudioMaterialSlate/AudioMaterialSlateTypes.h"
+#include "AudioMaterialSlateStyles.h"
 #include "Components/Widget.h"
 #include "Delegates/Delegate.h"
 #include "AudioMaterialMeter.generated.h"
@@ -23,8 +23,6 @@ class AUDIOWIDGETS_API UAudioMaterialMeter : public UWidget
 
 public:
 
-	UAudioMaterialMeter();
-
 	/** The meter's style */
 	UPROPERTY(EditAnywhere, BlueprintReadWrite, Category = "Style", meta = (DisplayName = "Style", ShowOnlyInnerProperties))
 	FAudioMaterialMeterStyle WidgetStyle;
diff --git a/Engine/Plugins/Runtime/AudioWidgets/Source/AudioWidgets/Public/AudioMaterialSlate/AudioMaterialSlateStyles.h b/Engine/Plugins/Runtime/AudioWidgets/Source/AudioWidgets/Public/AudioMaterialSlate/AudioMaterialSlateStyles.h
new file mode 100644
--- /dev/null
+++ b/Engine/Plugins/Runtime/AudioWidgets/Source/AudioWidgets/Public/AudioMaterialSlate/AudioMaterialSlateStyles.h
@@ -0,0 +1,220 @@
+// Copyright Epic Games, Inc. All Rights Reserved.
+
+#pragma once
+
+#include "Materials/MaterialInstanceDynamic.h"
+#include "Materials/MaterialInterface.h"
+#include "Math/Color.h"
+#include "Styling/ISlateStyle.h"
+#include "Styling/SlateWidgetStyle.h"
+#include "Styling/SlateWidgetStyleAsset.h"
+#include "Styling/SlateWidgetStyleContainerBase.h"
+#include "AudioMaterialSlateStyles.generated.h"
+
+
+/**
+ *Base for the appearance of an Audio Material Slates 
+ */
+USTRUCT(BlueprintType)
+struct AUDIOWIDGETS_API FAudioMaterialStyle : public FSlateWidgetStyle
+{
+	GENERATED_USTRUCT_BODY()
+
+public:
+
+	FAudioMaterialStyle();
+
+	/** Material used to render the Slate */
+	UPROPERTY(EditAnywhere, BlueprintReadWrite, meta = (DisplayPriority = 0), Category = "Style")
+	TObjectPtr<UMaterialInterface> Material = nullptr;
+
+	/** Desired Draw size of the rendered material*/
+	UPROPERTY(EditAnywhere, BlueprintReadWrite, meta = (DisplayPriority = 1), Category = "Style")
+	FVector2f DesiredSize;
+
+public:
+
+	void CreateDynamicMaterial(UObject* InOuter) const;
+	TObjectPtr<UMaterialInstanceDynamic> GetDynamicMaterial() const;
+
+private:
+
+	UPROPERTY()
+	mutable TObjectPtr<UMaterialInstanceDynamic> DynamicMaterial;
+};
+
+/**
+ *Represents the appearance of an Audio Material Button 
+ */
+USTRUCT(BlueprintType)
+struct AUDIOWIDGETS_API FAudioMaterialButtonStyle : public FAudioMaterialStyle
+{
+	GENERATED_USTRUCT_BODY()
+
+	FAudioMaterialButtonStyle();
+	virtual ~FAudioMaterialButtonStyle();
+
+	// FSlateWidgetStyle
+	virtual void GetResources(TArray<const FSlateBrush*>& OutBrushes) const override;
+	static const FName TypeName;
+	virtual const FName GetTypeName() const override { return TypeName; };
+	static const FAudioMaterialButtonStyle& GetDefault();
+
+	UPROPERTY(EditAnywhere, BlueprintReadWrite, Category = "Style")
+	FLinearColor ButtonMainColor;
+
+	UPROPERTY(EditAnywhere, BlueprintReadWrite, Category = "Style")
+	FLinearColor ButtonShadowColor;
+
+	UPROPERTY(EditAnywhere, BlueprintReadWrite, Category = "Style")
+	FLinearColor ButtonAccentColor;	
+	
+	UPROPERTY(EditAnywhere, BlueprintReadWrite, Category = "Style")
+	FLinearColor ButtonPressedMainColor;	
+	
+	UPROPERTY(EditAnywhere, BlueprintReadWrite, Category = "Style")
+	FLinearColor ButtonShadowMainColor;	
+	
+	UPROPERTY(EditAnywhere, BlueprintReadWrite, Category = "Style")
+	FLinearColor ButtonPressedOutlineColor;
+
+};
+
+/**
+ *Represents the appearance of an Audio Material Slider 
+ */
+USTRUCT(BlueprintType)
+struct AUDIOWIDGETS_API FAudioMaterialSliderStyle : public FAudioMaterialStyle
+{
+	GENERATED_USTRUCT_BODY()
+
+	FAudioMaterialSliderStyle();
+	virtual ~FAudioMaterialSliderStyle();
+
+	// FSlateWidgetStyle
+	virtual void GetResources(TArray<const FSlateBrush*>& OutBrushes) const override;
+	static const FName TypeName;
+	virtual const FName GetTypeName() const override { return TypeName; };
+	static const FAudioMaterialSliderStyle& GetDefault();
+
+	UPROPERTY(EditAnywhere, BlueprintReadWrite, Category = "Style")
+	FLinearColor BarMainColor;	
+	
+	UPROPERTY(EditAnywhere, BlueprintReadWrite, Category = "Style")
+	FLinearColor BarShadowColor;	
+	
+	UPROPERTY(EditAnywhere, BlueprintReadWrite, Category = "Style")
+	FLinearColor BarAccentColor;
+
+	UPROPERTY(EditAnywhere, BlueprintReadWrite, Category = "Style")
+	FLinearColor HandleMainColor;	
+	
+	UPROPERTY(EditAnywhere, BlueprintReadWrite, Category = "Style")
+	FLinearColor HandleOutlineColor;
+
+};
+
+/**
+ *Represents the appearance of an Audio Material Knob 
+ */
+USTRUCT(BlueprintType)
+struct AUDIOWIDGETS_API FAudioMaterialKnobStyle : public FAudioMaterialStyle
+{
+	GENERATED_USTRUCT_BODY()
+
+	FAudioMaterialKnobStyle();
+	virtual ~FAudioMaterialKnobStyle();
+
+	// FSlateWidgetStyle
+	virtual void GetResources(TArray<const FSlateBrush*>& OutBrushes) const override;
+	static const FName TypeName;
+	virtual const FName GetTypeName() const override { return TypeName; };
+	static const FAudioMaterialKnobStyle& GetDefault();
+
+	UPROPERTY(EditAnywhere, BlueprintReadWrite, Category = "Style")
+	FLinearColor KnobMainColor;
+
+	UPROPERTY(EditAnywhere, BlueprintReadWrite, Category = "Style")
+	FLinearColor KnobAccentColor;
+	
+	UPROPERTY(EditAnywhere, BlueprintReadWrite, Category = "Style")
+	FLinearColor KnobIndicatorColor;
+	
+	UPROPERTY(EditAnywhere, BlueprintReadWrite, Category = "Style")
+	FLinearColor KnobBarColor;		
+	
+	UPROPERTY(EditAnywhere, BlueprintReadWrite, Category = "Style")
+	FLinearColor KnobBarShadowColor;	
+
+	UPROPERTY(EditAnywhere, BlueprintReadWrite, Category = "Style")
+	FLinearColor KnobBarFillMinColor;		
+	
+	UPROPERTY(EditAnywhere, BlueprintReadWrite, Category = "Style")
+	FLinearColor KnobBarFillMidColor;		
+	
+	UPROPERTY(EditAnywhere, BlueprintReadWrite, Category = "Style")
+	FLinearColor KnobBarFillMaxColor;	
+	
+	UPROPERTY(EditAnywhere, BlueprintReadWrite, Category = "Style")
+	FLinearColor KnobBarFillTintColor;		
+	
+};
+
+/**
+ *Represents the appearance of an Audio Material Meter
+ */
+USTRUCT(BlueprintType)
+struct AUDIOWIDGETS_API FAudioMaterialMeterStyle : public FAudioMaterialStyle
+{
+	GENERATED_USTRUCT_BODY()
+
+	FAudioMaterialMeterStyle();
+	virtual ~FAudioMaterialMeterStyle();
+
+	// FSlateWidgetStyle
+	virtual void GetResources(TArray<const FSlateBrush*>& OutBrushes) const override;
+	static const FName TypeName;
+	virtual const FName GetTypeName() const override { return TypeName; };
+	static const FAudioMaterialMeterStyle& GetDefault();
+
+	UPROPERTY(EditAnywhere, BlueprintReadWrite, Category = "Style")
+	FLinearColor MeterFillMinColor;
+
+	UPROPERTY(EditAnywhere, BlueprintReadWrite, Category = "Style")
+	FLinearColor MeterFillMidColor;
+
+	UPROPERTY(EditAnywhere, BlueprintReadWrite, Category = "Style")
+	FLinearColor MeterFillMaxColor;	
+	
+	UPROPERTY(EditAnywhere, BlueprintReadWrite, Category = "Style")
+	FLinearColor MeterOffFillColor;
+	
+};
+
+/**
+ *Represents the appearance of an Audio Material Envelope
+ */
+USTRUCT(BlueprintType)
+struct AUDIOWIDGETS_API FAudioMaterialEnvelopeStyle : public FAudioMaterialStyle
+{
+	GENERATED_USTRUCT_BODY()
+
+	FAudioMaterialEnvelopeStyle();
+	virtual ~FAudioMaterialEnvelopeStyle();
+
+	// FSlateWidgetStyle
+	virtual void GetResources(TArray<const FSlateBrush*>& OutBrushes) const override;
+	static const FName TypeName;
+	virtual const FName GetTypeName() const override { return TypeName; };
+	static const FAudioMaterialEnvelopeStyle& GetDefault();
+
+	UPROPERTY(EditAnywhere, BlueprintReadWrite, Category = "Style")
+	FLinearColor CurveColor;	
+	
+	UPROPERTY(EditAnywhere, BlueprintReadWrite, Category = "Style")
+	FLinearColor BackgroundColor;	
+	
+	UPROPERTY(EditAnywhere, BlueprintReadWrite, Category = "Style")
+	FLinearColor OutlineColor;
+	
+};
diff --git a/Engine/Plugins/Runtime/AudioWidgets/Source/AudioWidgets/Public/AudioMaterialSlate/AudioMaterialSlateTypes.h b/Engine/Plugins/Runtime/AudioWidgets/Source/AudioWidgets/Public/AudioMaterialSlate/AudioMaterialSlateTypes.h
deleted file mode 100644
diff --git a/Engine/Plugins/Runtime/AudioWidgets/Source/AudioWidgets/Public/AudioMaterialSlate/AudioMaterialSlider.h b/Engine/Plugins/Runtime/AudioWidgets/Source/AudioWidgets/Public/AudioMaterialSlate/AudioMaterialSlider.h
--- a/Engine/Plugins/Runtime/AudioWidgets/Source/AudioWidgets/Public/AudioMaterialSlate/AudioMaterialSlider.h
+++ b/Engine/Plugins/Runtime/AudioWidgets/Source/AudioWidgets/Public/AudioMaterialSlate/AudioMaterialSlider.h
@@ -2,7 +2,7 @@
 
 #pragma once
 
-#include "AudioMaterialSlate/AudioMaterialSlateTypes.h"
+#include "AudioMaterialSlateStyles.h"
 #include "Components/Widget.h"
 #include "Delegates/Delegate.h"
 #include "AudioMaterialSlider.generated.h"
@@ -24,8 +24,6 @@ class AUDIOWIDGETS_API UAudioMaterialSlider : public UWidget
 
 public:
 
-	UAudioMaterialSlider();
-
 	/** The slider's style */
 	UPROPERTY(EditAnywhere, BlueprintReadWrite, Category = "Style", meta = (DisplayName = "Style", ShowOnlyInnerProperties))
 	FAudioMaterialSliderStyle WidgetStyle;
diff --git a/Engine/Plugins/Runtime/AudioWidgets/Source/AudioWidgets/Public/AudioMaterialSlate/AudioMaterialStyleContainers.h b/Engine/Plugins/Runtime/AudioWidgets/Source/AudioWidgets/Public/AudioMaterialSlate/AudioMaterialStyleContainers.h
deleted file mode 100644
diff --git a/Engine/Plugins/Runtime/AudioWidgets/Source/AudioWidgets/Public/AudioMaterialSlate/SAudioMaterialButton.h b/Engine/Plugins/Runtime/AudioWidgets/Source/AudioWidgets/Public/AudioMaterialSlate/SAudioMaterialButton.h
--- a/Engine/Plugins/Runtime/AudioWidgets/Source/AudioWidgets/Public/AudioMaterialSlate/SAudioMaterialButton.h
+++ b/Engine/Plugins/Runtime/AudioWidgets/Source/AudioWidgets/Public/AudioMaterialSlate/SAudioMaterialButton.h
@@ -2,6 +2,7 @@
 
 #pragma once
 
+#include "AudioMaterialSlate/AudioMaterialSlateStyles.h"
 #include "Framework/SlateDelegates.h"
 #include "Styling/ISlateStyle.h"
 #include "Styling/SlateWidgetStyleAsset.h"
diff --git a/Engine/Plugins/Runtime/AudioWidgets/Source/AudioWidgets/Public/AudioMaterialSlate/SAudioMaterialEnvelope.h b/Engine/Plugins/Runtime/AudioWidgets/Source/AudioWidgets/Public/AudioMaterialSlate/SAudioMaterialEnvelope.h
--- a/Engine/Plugins/Runtime/AudioWidgets/Source/AudioWidgets/Public/AudioMaterialSlate/SAudioMaterialEnvelope.h
+++ b/Engine/Plugins/Runtime/AudioWidgets/Source/AudioWidgets/Public/AudioMaterialSlate/SAudioMaterialEnvelope.h
@@ -2,7 +2,7 @@
 
 #pragma once
 
-#include "AudioMaterialSlate/AudioMaterialSlateTypes.h"
+#include "AudioMaterialSlate/AudioMaterialSlateStyles.h"
 #include "AudioMaterialEnvelopeSettings.h"
 #include "Framework/SlateDelegates.h"
 #include "Styling/ISlateStyle.h"
diff --git a/Engine/Plugins/Runtime/AudioWidgets/Source/AudioWidgets/Public/AudioMaterialSlate/SAudioMaterialKnob.h b/Engine/Plugins/Runtime/AudioWidgets/Source/AudioWidgets/Public/AudioMaterialSlate/SAudioMaterialKnob.h
--- a/Engine/Plugins/Runtime/AudioWidgets/Source/AudioWidgets/Public/AudioMaterialSlate/SAudioMaterialKnob.h
+++ b/Engine/Plugins/Runtime/AudioWidgets/Source/AudioWidgets/Public/AudioMaterialSlate/SAudioMaterialKnob.h
@@ -2,7 +2,7 @@
 
 #pragma once
 
-#include "AudioMaterialSlate/AudioMaterialSlateTypes.h"
+#include "AudioMaterialSlate/AudioMaterialSlateStyles.h"
 #include "Framework/SlateDelegates.h"
 #include "Styling/ISlateStyle.h"
 #include "Styling/SlateWidgetStyleAsset.h"
diff --git a/Engine/Plugins/Runtime/AudioWidgets/Source/AudioWidgets/Public/AudioMaterialSlate/SAudioMaterialMeter.h b/Engine/Plugins/Runtime/AudioWidgets/Source/AudioWidgets/Public/AudioMaterialSlate/SAudioMaterialMeter.h
--- a/Engine/Plugins/Runtime/AudioWidgets/Source/AudioWidgets/Public/AudioMaterialSlate/SAudioMaterialMeter.h
+++ b/Engine/Plugins/Runtime/AudioWidgets/Source/AudioWidgets/Public/AudioMaterialSlate/SAudioMaterialMeter.h
@@ -2,7 +2,7 @@
 
 #pragma once
 
-#include "AudioMaterialSlate/AudioMaterialSlateTypes.h"
+#include "AudioMaterialSlate/AudioMaterialSlateStyles.h"
 #include "Framework/SlateDelegates.h"
 #include "Styling/ISlateStyle.h"
 #include "Styling/SlateWidgetStyleAsset.h"
diff --git a/Engine/Plugins/Runtime/AudioWidgets/Source/AudioWidgets/Public/AudioMaterialSlate/SAudioMaterialSlider.h b/Engine/Plugins/Runtime/AudioWidgets/Source/AudioWidgets/Public/AudioMaterialSlate/SAudioMaterialSlider.h
--- a/Engine/Plugins/Runtime/AudioWidgets/Source/AudioWidgets/Public/AudioMaterialSlate/SAudioMaterialSlider.h
+++ b/Engine/Plugins/Runtime/AudioWidgets/Source/AudioWidgets/Public/AudioMaterialSlate/SAudioMaterialSlider.h
@@ -2,7 +2,7 @@
 
 #pragma once
 
-#include "AudioMaterialSlate/AudioMaterialSlateTypes.h"
+#include "AudioMaterialSlate/AudioMaterialSlateStyles.h"
 #include "Framework/SlateDelegates.h"
 #include "Styling/ISlateStyle.h"
 #include "Styling/SlateWidgetStyleAsset.h"
