AudioMaterialKnob to Metasound editor float node. Will in the end replace the current Knob, but for now, this is behind CVar for testing.
#rnx

diff --git a/Engine/Plugins/Runtime/AudioWidgets/Source/AudioWidgets/Private/AudioMaterialSlate/AudioMaterialSlateTypes.cpp b/Engine/Plugins/Runtime/AudioWidgets/Source/AudioWidgets/Private/AudioMaterialSlate/AudioMaterialSlateTypes.cpp
--- a/Engine/Plugins/Runtime/AudioWidgets/Source/AudioWidgets/Private/AudioMaterialSlate/AudioMaterialSlateTypes.cpp
+++ b/Engine/Plugins/Runtime/AudioWidgets/Source/AudioWidgets/Private/AudioMaterialSlate/AudioMaterialSlateTypes.cpp
@@ -13,9 +13,10 @@ TObjectPtr<UMaterialInstanceDynamic> FAudioMaterialWidgetStyle::GetDynamicMateri
 	return DynamicMaterial;
 }
 
-void FAudioMaterialWidgetStyle::CreateDynamicMaterial(UObject* InOuter) const
+UMaterialInstanceDynamic* FAudioMaterialWidgetStyle::CreateDynamicMaterial(UObject* InOuter) const
 {
 	DynamicMaterial = UMaterialInstanceDynamic::Create(Material, InOuter);
+	return DynamicMaterial;
 }
 
 FAudioMaterialButtonStyle::FAudioMaterialButtonStyle()
@@ -75,7 +76,7 @@ FAudioMaterialKnobStyle::FAudioMaterialKnobStyle()
 	, KnobBarFillMaxColor(FLinearColor::White)
 	, KnobBarFillTintColor(FLinearColor::White)
 {
-	DesiredSize = FVector2f(256.f,256.f);
+	DesiredSize = FVector2f(128.f,128.f);
 }
 
 const FName FAudioMaterialKnobStyle::TypeName(TEXT("FAudioMaterialKnobStyle"));
diff --git a/Engine/Plugins/Runtime/AudioWidgets/Source/AudioWidgets/Private/AudioMaterialSlate/SAudioMaterialKnob.cpp b/Engine/Plugins/Runtime/AudioWidgets/Source/AudioWidgets/Private/AudioMaterialSlate/SAudioMaterialKnob.cpp
--- a/Engine/Plugins/Runtime/AudioWidgets/Source/AudioWidgets/Private/AudioMaterialSlate/SAudioMaterialKnob.cpp
+++ b/Engine/Plugins/Runtime/AudioWidgets/Source/AudioWidgets/Private/AudioMaterialSlate/SAudioMaterialKnob.cpp
@@ -3,6 +3,7 @@
 
 #include "AudioMaterialSlate/SAudioMaterialKnob.h"
 #include "AudioMaterialSlate/AudioMaterialSlateTypes.h"
+#include "AudioWidgetsStyle.h"
 #include "SlateOptMacros.h"
 #include "Components/Widget.h"
 
@@ -12,9 +13,21 @@ void SAudioMaterialKnob::Construct(const FArguments& InArgs)
 	Owner = InArgs._Owner;
 	AudioMaterialKnobStyle = InArgs._AudioMaterialKnobStyle;
 	OnValueChanged = InArgs._OnFloatValueChanged;
+	OnMouseCaptureBegin = InArgs._OnMouseCaptureBegin;
+	OnMouseCaptureEnd = InArgs._OnMouseCaptureEnd;
+
+	//For now check if owner is not a widget -> use the default style.
+	if (!Cast<UWidget>(Owner))
+	{
+		AudioMaterialKnobStyle = &FAudioWidgetsStyle::Get().GetWidgetStyle<FAudioMaterialKnobStyle>("AudioMaterialKnob.Style");
+	}
 
 	ApplyNewMaterial();
-	CommitValue(InArgs._Value.Get());
+
+	if (InArgs._Value.IsSet())
+	{
+		CommitValue(InArgs._Value.Get());
+	}
 }
 END_SLATE_FUNCTION_BUILD_OPTIMIZATION
 
@@ -23,54 +36,99 @@ void SAudioMaterialKnob::SetValue(float InValue)
 	CommitValue(InValue);
 }
 
-void SAudioMaterialKnob::ApplyNewMaterial()
+UMaterialInstanceDynamic* SAudioMaterialKnob::ApplyNewMaterial()
 {
 	if (AudioMaterialKnobStyle)
 	{
-		AudioMaterialKnobStyle->CreateDynamicMaterial(Owner.Get());
+		DynamicMaterial = AudioMaterialKnobStyle->CreateDynamicMaterial(Owner.Get());
 	}
+
+	return DynamicMaterial.Get();
+}
+const float SAudioMaterialKnob::GetOutputValue(const float InSliderValue)
+{
+	return FMath::GetMappedRangeValueClamped(NormalizedLinearSliderRange, OutputRange, InSliderValue);
+}
+
+const float SAudioMaterialKnob::GetSliderValue(const float OutputValue)
+{
+	return FMath::GetMappedRangeValueClamped(OutputRange, NormalizedLinearSliderRange, OutputValue);
+}
+
+void SAudioMaterialKnob::SetSliderValue(float InSliderValue)
+{
+	ValueAttribute.Set(InSliderValue);
+}
+
+void SAudioMaterialKnob::SetOutputRange(const FVector2D Range)
+{
+	OutputRange = Range;
+	// if Range.Y < Range.X, set Range.X to Range.Y
+	OutputRange.X = FMath::Min(Range.X, Range.Y);
+
+	const float OutputValue = GetOutputValue(ValueAttribute.Get());
+	const float ClampedOutputValue = FMath::Clamp(OutputValue, OutputRange.X, OutputRange.Y);
+	const float ClampedSliderValue = GetSliderValue(ClampedOutputValue);
+	SetSliderValue(ClampedSliderValue);
+}
+
+void SAudioMaterialKnob::SetDesiredSizeOverride(const FVector2D Size)
+{
+	SetAttribute(DesiredSizeOverride, TAttribute<TOptional<FVector2D>>(Size), EInvalidateWidgetReason::Layout);
 }
 
 int32 SAudioMaterialKnob::OnPaint(const FPaintArgs& Args, const FGeometry& AllottedGeometry, const FSlateRect& MyCullingRect, FSlateWindowElementList& OutDrawElements, int32 LayerId, const FWidgetStyle& InWidgetStyle, bool bParentEnabled) const
 {
 	if (AudioMaterialKnobStyle)
 	{
-		UMaterialInstanceDynamic* DynamicMaterial = AudioMaterialKnobStyle->GetDynamicMaterial();
-
-		if (IsValid(DynamicMaterial))
+		if (DynamicMaterial.IsValid())
 		{
 			const float KnobPercent = ValueAttribute.Get();
 
 			//TODO unify material parameter names
-			DynamicMaterial->SetVectorParameterValue(FName("Color_2"), AudioMaterialKnobStyle->KnobMainColor);
-			DynamicMaterial->SetVectorParameterValue(FName("Color_1"), AudioMaterialKnobStyle->KnobAccentColor);
-
-			DynamicMaterial->SetVectorParameterValue(FName("BarColor"), AudioMaterialKnobStyle->KnobBarColor);
-			DynamicMaterial->SetVectorParameterValue(FName("BarShadowColor"), AudioMaterialKnobStyle->KnobBarShadowColor);
-			DynamicMaterial->SetVectorParameterValue(FName("LEDglowMax"), AudioMaterialKnobStyle->KnobBarFillMaxColor);
-			DynamicMaterial->SetVectorParameterValue(FName("Led_Max"), AudioMaterialKnobStyle->KnobBarFillMaxColor);
-			DynamicMaterial->SetVectorParameterValue(FName("LEDglowMid"), AudioMaterialKnobStyle->KnobBarFillMidColor);
-			DynamicMaterial->SetVectorParameterValue(FName("Led_Mid"), AudioMaterialKnobStyle->KnobBarFillMidColor);
-			DynamicMaterial->SetVectorParameterValue(FName("LEDglowMin"), AudioMaterialKnobStyle->KnobBarFillMinColor);
-			DynamicMaterial->SetVectorParameterValue(FName("LED_Min"), AudioMaterialKnobStyle->KnobBarFillMinColor);
-			DynamicMaterial->SetVectorParameterValue(FName("LineColor"), AudioMaterialKnobStyle->KnobIndicatorColor);
-			DynamicMaterial->SetVectorParameterValue(FName("DotColor"), AudioMaterialKnobStyle->KnobIndicatorColor);
-			DynamicMaterial->SetVectorParameterValue(FName("LedTint"), AudioMaterialKnobStyle->KnobBarFillTintColor);
-
-			DynamicMaterial->SetScalarParameterValue(FName("VALUE"), FMath::Clamp(KnobPercent, 0.f, 1.f));
-
-			DynamicMaterial->SetScalarParameterValue(FName("LocalWidth"), AllottedGeometry.GetLocalSize().X);
-			DynamicMaterial->SetScalarParameterValue(FName("LocalHeigth"), AllottedGeometry.GetLocalSize().Y);			
+			DynamicMaterial.Get()->SetVectorParameterValue(FName("Color_2"), AudioMaterialKnobStyle->KnobMainColor);
+			DynamicMaterial.Get()->SetVectorParameterValue(FName("Color_1"), AudioMaterialKnobStyle->KnobAccentColor);
+
+			DynamicMaterial.Get()->SetVectorParameterValue(FName("BarColor"), AudioMaterialKnobStyle->KnobBarColor);
+			DynamicMaterial.Get()->SetVectorParameterValue(FName("BarShadowColor"), AudioMaterialKnobStyle->KnobBarShadowColor);
+			DynamicMaterial.Get()->SetVectorParameterValue(FName("LEDglowMax"), AudioMaterialKnobStyle->KnobBarFillMaxColor);
+			DynamicMaterial.Get()->SetVectorParameterValue(FName("Led_Max"), AudioMaterialKnobStyle->KnobBarFillMaxColor);
+			DynamicMaterial.Get()->SetVectorParameterValue(FName("LEDglowMid"), AudioMaterialKnobStyle->KnobBarFillMidColor);
+			DynamicMaterial.Get()->SetVectorParameterValue(FName("Led_Mid"), AudioMaterialKnobStyle->KnobBarFillMidColor);
+			DynamicMaterial.Get()->SetVectorParameterValue(FName("LEDglowMin"), AudioMaterialKnobStyle->KnobBarFillMinColor);
+			DynamicMaterial.Get()->SetVectorParameterValue(FName("LED_Min"), AudioMaterialKnobStyle->KnobBarFillMinColor);
+			DynamicMaterial.Get()->SetVectorParameterValue(FName("LineColor"), AudioMaterialKnobStyle->KnobIndicatorColor);
+			DynamicMaterial.Get()->SetVectorParameterValue(FName("DotColor"), AudioMaterialKnobStyle->KnobIndicatorColor);
+			DynamicMaterial.Get()->SetVectorParameterValue(FName("LedTint"), AudioMaterialKnobStyle->KnobBarFillTintColor);
+
+			DynamicMaterial.Get()->SetScalarParameterValue(FName("VALUE"), FMath::Clamp(KnobPercent, 0.f, 1.f));
+
+			DynamicMaterial.Get()->SetScalarParameterValue(FName("LocalWidth"), AllottedGeometry.GetLocalSize().X);
+			DynamicMaterial.Get()->SetScalarParameterValue(FName("LocalHeigth"), AllottedGeometry.GetLocalSize().Y);
+
+			const bool bEnabled = ShouldBeEnabled(bParentEnabled);
+			const ESlateDrawEffect DrawEffects = bEnabled ? ESlateDrawEffect::None : ESlateDrawEffect::DisabledEffect;
+
+			const FLinearColor FinalColorAndOpacity(InWidgetStyle.GetColorAndOpacityTint());
+			
+			const float AllottedWidth = AllottedGeometry.GetLocalSize().X;
+			const float AllottedHeight = AllottedGeometry.GetLocalSize().Y;
+			
+			const float SliderRadius = FMath::Min(AllottedWidth, AllottedHeight) * 0.5f;
+			const FVector2D SliderMidPoint(AllottedGeometry.GetLocalSize() * 0.5f);
+			const FVector2D SliderDiameter(SliderRadius * 2);
+
+			FSlateBrush Brush;
+			Brush.SetResourceObject(DynamicMaterial.Get());
+			FSlateDrawElement::MakeBox(OutDrawElements, LayerId++, AllottedGeometry.ToPaintGeometry(SliderDiameter, FSlateLayoutTransform(SliderMidPoint-SliderRadius)), &Brush, DrawEffects, FinalColorAndOpacity);
+		}
+		else
+		{
+			if (AudioMaterialKnobStyle)
+			{
+				DynamicMaterial = AudioMaterialKnobStyle->CreateDynamicMaterial(Owner.Get());
+			}
 		}
-
-		const bool bEnabled = ShouldBeEnabled(bParentEnabled);
-		const ESlateDrawEffect DrawEffects = bEnabled ? ESlateDrawEffect::None : ESlateDrawEffect::DisabledEffect;
-
-		const FLinearColor FinalColorAndOpacity(InWidgetStyle.GetColorAndOpacityTint());
-
-		FSlateBrush Brush;
-		Brush.SetResourceObject(DynamicMaterial);
-		FSlateDrawElement::MakeBox(OutDrawElements, LayerId++, AllottedGeometry.ToPaintGeometry(),&Brush, DrawEffects, FinalColorAndOpacity);
 	}
 
 	return LayerId;
@@ -78,6 +136,11 @@ int32 SAudioMaterialKnob::OnPaint(const FPaintArgs& Args, const FGeometry& Allot
 
 FVector2D SAudioMaterialKnob::ComputeDesiredSize(float) const
 {
+	if (DesiredSizeOverride.Get().IsSet())
+	{
+		return DesiredSizeOverride.Get().GetValue();
+	}
+	
 	if (AudioMaterialKnobStyle)
 	{
 		return FVector2D(AudioMaterialKnobStyle->DesiredSize);
@@ -90,8 +153,9 @@ FReply SAudioMaterialKnob::OnMouseMove(const FGeometry& MyGeometry, const FPoint
 {
 	if (this->HasMouseCapture())
 	{
-		int32 CurrentYValue = MouseEvent.GetLastScreenSpacePosition().Y;
+		SetCursor(EMouseCursor::GrabHandClosed);
 
+		int32 CurrentYValue = MouseEvent.GetLastScreenSpacePosition().Y;
 		const float MouseSpeed = 0.2f;
 
 		float ValueDelta = (float)(MouseDownPosition.Y - CurrentYValue) / PixelDelta * MouseSpeed;
@@ -108,8 +172,11 @@ FReply SAudioMaterialKnob::OnMouseButtonDown(const FGeometry& MyGeometry, const
 {
 	if ((MouseEvent.GetEffectingButton() == EKeys::LeftMouseButton))
 	{
+		CachedCursor = GetCursor().Get(EMouseCursor::Default);
+
 		MouseDownPosition = MouseEvent.GetScreenSpacePosition();
 		MouseDownValue = ValueAttribute.Get();
+		OnMouseCaptureBegin.ExecuteIfBound();
 
 		FReply Reply = FReply::Handled().CaptureMouse(SharedThis(this));
 		return Reply;
@@ -122,6 +189,8 @@ FReply SAudioMaterialKnob::OnMouseButtonUp(const FGeometry& MyGeometry, const FP
 {
 	if ((MouseEvent.GetEffectingButton() == EKeys::LeftMouseButton) && this->HasMouseCapture())
 	{
+		SetCursor(CachedCursor);
+		OnMouseCaptureEnd.ExecuteIfBound();
 		return FReply::Handled().ReleaseMouseCapture();
 	}
 
@@ -131,5 +200,6 @@ FReply SAudioMaterialKnob::OnMouseButtonUp(const FGeometry& MyGeometry, const FP
 void SAudioMaterialKnob::CommitValue(float NewValue)
 {
 	ValueAttribute.Set(NewValue);
+	Invalidate(EInvalidateWidgetReason::Paint);
 	OnValueChanged.ExecuteIfBound(NewValue);
-}
+}
\ No newline at end of file
diff --git a/Engine/Plugins/Runtime/AudioWidgets/Source/AudioWidgets/Public/AudioMaterialSlate/AudioMaterialSlateTypes.h b/Engine/Plugins/Runtime/AudioWidgets/Source/AudioWidgets/Public/AudioMaterialSlate/AudioMaterialSlateTypes.h
--- a/Engine/Plugins/Runtime/AudioWidgets/Source/AudioWidgets/Public/AudioMaterialSlate/AudioMaterialSlateTypes.h
+++ b/Engine/Plugins/Runtime/AudioWidgets/Source/AudioWidgets/Public/AudioMaterialSlate/AudioMaterialSlateTypes.h
@@ -31,12 +31,12 @@ public:
 
 public:
 
-	void CreateDynamicMaterial(UObject* InOuter) const;
+	UMaterialInstanceDynamic* CreateDynamicMaterial(UObject* InOuter) const;
 	TObjectPtr<UMaterialInstanceDynamic> GetDynamicMaterial() const;
 
 private:
 
-	UPROPERTY()
+	UPROPERTY(meta = (DeprecatedProperty, DeprecationMessage = "keep the reference instead as a TWeakObjectPtr in the AudioMaterialSlates.  This will be removed when all AudioMaterialSlates are updated."))
 	mutable TObjectPtr<UMaterialInstanceDynamic> DynamicMaterial;
 };
 
diff --git a/Engine/Plugins/Runtime/AudioWidgets/Source/AudioWidgets/Public/AudioMaterialSlate/SAudioMaterialKnob.h b/Engine/Plugins/Runtime/AudioWidgets/Source/AudioWidgets/Public/AudioMaterialSlate/SAudioMaterialKnob.h
--- a/Engine/Plugins/Runtime/AudioWidgets/Source/AudioWidgets/Public/AudioMaterialSlate/SAudioMaterialKnob.h
+++ b/Engine/Plugins/Runtime/AudioWidgets/Source/AudioWidgets/Public/AudioMaterialSlate/SAudioMaterialKnob.h
@@ -4,6 +4,7 @@
 
 #include "AudioMaterialSlate/AudioMaterialSlateTypes.h"
 #include "Framework/SlateDelegates.h"
+#include "SAudioInputWidget.h"
 #include "Styling/ISlateStyle.h"
 #include "Styling/SlateWidgetStyleAsset.h"
 #include "Widgets/SLeafWidget.h"
@@ -14,7 +15,7 @@ class UObject;
  * A simple slate that renders a knob in single material and modifies the material on value change.
  *
  */
-class AUDIOWIDGETS_API SAudioMaterialKnob : public SLeafWidget
+class AUDIOWIDGETS_API SAudioMaterialKnob : public SAudioInputWidget
 {
 public:
 	SLATE_BEGIN_ARGS(SAudioMaterialKnob)
@@ -26,12 +27,18 @@ public:
 	/**Value of the Knob*/
 	SLATE_ATTRIBUTE(float, Value)
 
-	/** The style used to draw the button. */
+	/** The style used to draw the knob. */
 	SLATE_STYLE_ARGUMENT(FAudioMaterialKnobStyle, AudioMaterialKnobStyle)
 
-	/** Called when the button's state changes. */
+	/** Called when the knob's state changes. */
 	SLATE_EVENT(FOnFloatValueChanged, OnFloatValueChanged)
 
+	/** Invoked when the mouse is pressed and a capture begins. */
+	SLATE_EVENT(FSimpleDelegate, OnMouseCaptureBegin)
+
+	/** Invoked when the mouse is released and a capture ends. */
+	SLATE_EVENT(FSimpleDelegate, OnMouseCaptureEnd)
+
 	SLATE_END_ARGS()
 
 	/** Constructs this widget with InArgs */
@@ -41,23 +48,58 @@ public:
 	void SetValue(float InValueAttribute);
 
 	/** Apply new material to be used to render the Slate.*/
-	void ApplyNewMaterial();
+	UMaterialInstanceDynamic* ApplyNewMaterial();
+
+	//SAudioInputWidget
+	virtual const float GetOutputValue(const float InSliderValue) override;
+	virtual const float GetSliderValue(const float OutputValue) override;
+
+	/**
+	 * Set the knob's linear (0-1 normalized) value.
+	 */
+	virtual void SetSliderValue(float InSliderValue) override;
+	virtual void SetOutputRange(const FVector2D Range) override;
+
+	virtual void SetDesiredSizeOverride(const FVector2D Size) override;
+
+	//These are pure virtual functions in the parent class and are implemented properly in this class later.
+	virtual void SetLabelBackgroundColor(FSlateColor InColor) override {};
+	virtual void SetUnitsText(const FText Units) override {};
+	virtual void SetUnitsTextReadOnly(const bool bIsReadOnly) override {};
+	virtual void SetShowUnitsText(const bool bShowUnitsText) override {};
+	//~SAudioInputWidget
 
 public:
 
+	// Holds a delegate that is executed when the knob's value changes.
 	FOnFloatValueChanged OnValueChanged;
 
+	// Holds a delegate that is executed when the mouse is pressed and a capture begins.
+	FSimpleDelegate OnMouseCaptureBegin;
+
+	// Holds a delegate that is executed when the mouse is let up and a capture ends.
+	FSimpleDelegate OnMouseCaptureEnd;
+
 protected:
 
-	// SWidget overrides
+	//SWidget
 	virtual int32 OnPaint(const FPaintArgs& Args, const FGeometry& AllottedGeometry, const FSlateRect& MyCullingRect, FSlateWindowElementList& OutDrawElements, int32 LayerId, const FWidgetStyle& InWidgetStyle, bool bParentEnabled) const override;
 	virtual FVector2D ComputeDesiredSize(float) const override;
 	virtual FReply OnMouseMove(const FGeometry& MyGeometry, const FPointerEvent& MouseEvent) override;
 	virtual FReply OnMouseButtonDown(const FGeometry& MyGeometry, const FPointerEvent& MouseEvent) override;
 	virtual FReply OnMouseButtonUp(const FGeometry& MyGeometry, const FPointerEvent& MouseEvent)override;
+	//~SWidget
+
+protected:
+
+	// Range for output 
+	FVector2D OutputRange = FVector2D(0.0f, 1.0f);
+	const FVector2D NormalizedLinearSliderRange = FVector2D(0.0f, 1.0f);
 
 private:
 
+	TAttribute<TOptional<FVector2D>> DesiredSizeOverride;
+
 	/**Commits new value*/
 	void CommitValue(float NewValue);
 
@@ -69,7 +111,10 @@ private:
 	// Holds the style for the Slate
 	const FAudioMaterialKnobStyle* AudioMaterialKnobStyle = nullptr;
 
-	//Current Value
+	// Holds the Modifiable Material that represent the Knob
+	mutable TWeakObjectPtr<UMaterialInstanceDynamic> DynamicMaterial;
+
+	//Holds the knobs current Value
 	TAttribute<float> ValueAttribute = 1.0;
 
 	// The position of the mouse when it pushed down and started rotating the knob
@@ -78,6 +123,9 @@ private:
 	// the value when the mouse was pushed down
 	float MouseDownValue = 0.f;
 
+	// Holds the initial cursor in case a custom cursor has been specified, so we can restore it after dragging the slider
+	EMouseCursor::Type CachedCursor;
+
 	// the max pixels to go to min or max value (clamped to 0 or 1) in one drag period
 	int32 PixelDelta = 50;
 
diff --git a/Engine/Plugins/Runtime/Metasound/Source/MetasoundEditor/Private/SMetasoundGraphNode.cpp b/Engine/Plugins/Runtime/Metasound/Source/MetasoundEditor/Private/SMetasoundGraphNode.cpp
--- a/Engine/Plugins/Runtime/Metasound/Source/MetasoundEditor/Private/SMetasoundGraphNode.cpp
+++ b/Engine/Plugins/Runtime/Metasound/Source/MetasoundEditor/Private/SMetasoundGraphNode.cpp
@@ -1,6 +1,7 @@
 // Copyright Epic Games, Inc. All Rights Reserved.
 #include "SMetasoundGraphNode.h"
 
+#include "AudioMaterialSlate/SAudioMaterialKnob.h"
 #include "AudioParameterControllerInterface.h"
 #include "Components/AudioComponent.h"
 #include "GraphEditorSettings.h"
@@ -55,11 +56,21 @@
 
 #define LOCTEXT_NAMESPACE "MetasoundEditor"
 
-
 namespace Metasound
 {
 	namespace Editor
 	{
+		namespace GraphNodePrivate
+		{
+			int32 UseAudioMaterialWidgets = 0;
+			FAutoConsoleVariableRef CVarUseAudioMaterialWidgets(
+				TEXT("au.MetaSound.Editor.UseAudioMaterialSlates"),
+				UseAudioMaterialWidgets,
+				TEXT("Are new AudioMaterialWidgets used for visualization in the Metasound editor, if implemented.\n")
+				TEXT("0: Disabled (default), !0: Enabled"),
+				ECVF_Default);
+		}// Metasound::Editor::GraphNodePrivate
+
 		SMetaSoundGraphNode::~SMetaSoundGraphNode()
 		{
 			// Clean up input widgets
@@ -822,29 +833,40 @@ namespace Metasound
 								}
 							};
 
-							// Create slider 
-							if (DefaultFloat->WidgetValueType == EMetasoundMemberDefaultWidgetValueType::Frequency)
+							if (Metasound::Editor::GraphNodePrivate::UseAudioMaterialWidgets)
 							{
-								SAssignNew(InputWidget, SAudioFrequencyRadialSlider)
-									.OnValueChanged_Lambda(OnValueChangedLambda)
+								SAssignNew(InputWidget, SAudioMaterialKnob)
+									.Owner(GraphMember->GetOwningGraph())
+									.OnFloatValueChanged_Lambda(OnValueChangedLambda)
 									.OnMouseCaptureBegin_Lambda(OnRadialSliderMouseCaptureBeginLambda)
 									.OnMouseCaptureEnd_Lambda(OnRadialSliderMouseCaptureEndLambda);
 							}
-							else if (DefaultFloat->WidgetValueType == EMetasoundMemberDefaultWidgetValueType::Volume)
-							{
-								SAssignNew(InputWidget, SAudioVolumeRadialSlider)
-									.OnValueChanged_Lambda(OnValueChangedLambda)
-									.OnMouseCaptureBegin_Lambda(OnRadialSliderMouseCaptureBeginLambda)
-									.OnMouseCaptureEnd_Lambda(OnRadialSliderMouseCaptureEndLambda);
-								StaticCastSharedPtr<SAudioVolumeRadialSlider>(InputWidget)->SetUseLinearOutput(DefaultFloat->VolumeWidgetUseLinearOutput);
-							}
 							else
 							{
-								SAssignNew(InputWidget, SAudioRadialSlider)
-									.OnValueChanged_Lambda(OnValueChangedLambda)
-									.OnMouseCaptureBegin_Lambda(OnRadialSliderMouseCaptureBeginLambda)
-									.OnMouseCaptureEnd_Lambda(OnRadialSliderMouseCaptureEndLambda);
-								InputWidget->SetShowUnitsText(false);
+								// Create slider 
+								if (DefaultFloat->WidgetValueType == EMetasoundMemberDefaultWidgetValueType::Frequency)
+								{
+									SAssignNew(InputWidget, SAudioFrequencyRadialSlider)
+										.OnValueChanged_Lambda(OnValueChangedLambda)
+										.OnMouseCaptureBegin_Lambda(OnRadialSliderMouseCaptureBeginLambda)
+										.OnMouseCaptureEnd_Lambda(OnRadialSliderMouseCaptureEndLambda);
+								}
+								else if (DefaultFloat->WidgetValueType == EMetasoundMemberDefaultWidgetValueType::Volume)
+								{
+									SAssignNew(InputWidget, SAudioVolumeRadialSlider)
+										.OnValueChanged_Lambda(OnValueChangedLambda)
+										.OnMouseCaptureBegin_Lambda(OnRadialSliderMouseCaptureBeginLambda)
+										.OnMouseCaptureEnd_Lambda(OnRadialSliderMouseCaptureEndLambda);
+									StaticCastSharedPtr<SAudioVolumeRadialSlider>(InputWidget)->SetUseLinearOutput(DefaultFloat->VolumeWidgetUseLinearOutput);
+								}
+								else
+								{
+									SAssignNew(InputWidget, SAudioRadialSlider)
+										.OnValueChanged_Lambda(OnValueChangedLambda)
+										.OnMouseCaptureBegin_Lambda(OnRadialSliderMouseCaptureBeginLambda)
+										.OnMouseCaptureEnd_Lambda(OnRadialSliderMouseCaptureEndLambda);
+									InputWidget->SetShowUnitsText(false);
+								}
 							}
 							// Only vertical layout for radial slider
 							SAssignNew(OuterContentBox, SVerticalBox)
@@ -863,7 +885,10 @@ namespace Metasound
 								[
 									InputWidget.ToSharedRef()
 								];
-							InputWidget->SetDesiredSizeOverride(RadialSliderDesiredSize);
+							if (!Metasound::Editor::GraphNodePrivate::UseAudioMaterialWidgets)
+							{
+								InputWidget->SetDesiredSizeOverride(RadialSliderDesiredSize);
+							}
 						}
 
 						InputWidget->SetOutputRange(DefaultFloat->GetRange());
