AudioMaterialSlider Metasound editor float node
#rnx

#virtualized

diff --git a/Engine/Plugins/Runtime/AudioWidgets/Source/AudioWidgets/Private/AudioMaterialSlate/AudioMaterialSlateTypes.cpp b/Engine/Plugins/Runtime/AudioWidgets/Source/AudioWidgets/Private/AudioMaterialSlate/AudioMaterialSlateTypes.cpp
--- a/Engine/Plugins/Runtime/AudioWidgets/Source/AudioWidgets/Private/AudioMaterialSlate/AudioMaterialSlateTypes.cpp
+++ b/Engine/Plugins/Runtime/AudioWidgets/Source/AudioWidgets/Private/AudioMaterialSlate/AudioMaterialSlateTypes.cpp
@@ -43,13 +43,14 @@ const FAudioMaterialButtonStyle& FAudioMaterialButtonStyle::GetDefault()
 }
 
 FAudioMaterialSliderStyle::FAudioMaterialSliderStyle()
-	: BarMainColor(FLinearColor::White)
-	, BarShadowColor(FLinearColor::White)
-	, BarAccentColor(FLinearColor::White)
-	, HandleMainColor(FLinearColor::White)
-	, HandleOutlineColor(FLinearColor::White)
-{
-	DesiredSize = FVector2f(128.f, 512.f);
+	: BarMainColor(FLinearColor(0.008f, 0.008f, 0.008f,1.f))
+	, BarShadowColor(FLinearColor::Black)
+	, BarAccentColor(FLinearColor(0.005f, 0.005f, 0.005f, 1.f))
+	, HandleMainColor(FLinearColor(0.2f, 0.2f, 0.25f, 1.f))
+	, HandleOutlineColor(FLinearColor(0.15f, 0.15f, 0.15f, 1.f))
+	, TextBoxStyle(FAudioTextBoxStyle::GetDefault())
+{
+	DesiredSize = FVector2f(30.f, 250.f);
 }
 
 const FName FAudioMaterialSliderStyle::TypeName(TEXT("FAudioMaterialSliderStyle"));
diff --git a/Engine/Plugins/Runtime/AudioWidgets/Source/AudioWidgets/Private/AudioMaterialSlate/AudioMaterialSlider.cpp b/Engine/Plugins/Runtime/AudioWidgets/Source/AudioWidgets/Private/AudioMaterialSlate/AudioMaterialSlider.cpp
--- a/Engine/Plugins/Runtime/AudioWidgets/Source/AudioWidgets/Private/AudioMaterialSlate/AudioMaterialSlider.cpp
+++ b/Engine/Plugins/Runtime/AudioWidgets/Source/AudioWidgets/Private/AudioMaterialSlate/AudioMaterialSlider.cpp
@@ -32,6 +32,7 @@ void UAudioMaterialSlider::SynchronizeProperties()
 	}
 
 	Slider->SetValue(Value);
+	Slider->SetOrientation(Orientation);
 	Slider->ApplyNewMaterial();
 }
 
@@ -65,6 +66,7 @@ TSharedRef<SWidget> UAudioMaterialSlider::RebuildWidget()
 {
 	Slider = SNew(SAudioMaterialSlider)
 		.Owner(this)
+		.Orientation(Orientation)
 		.AudioMaterialSliderStyle(&WidgetStyle)
 		.OnValueChanged(BIND_UOBJECT_DELEGATE(FOnFloatValueChanged, HandleOnValueChanged));
 
diff --git a/Engine/Plugins/Runtime/AudioWidgets/Source/AudioWidgets/Private/AudioMaterialSlate/SAudioMaterialLabeledSlider.cpp b/Engine/Plugins/Runtime/AudioWidgets/Source/AudioWidgets/Private/AudioMaterialSlate/SAudioMaterialLabeledSlider.cpp
new file mode 100644
--- /dev/null
+++ b/Engine/Plugins/Runtime/AudioWidgets/Source/AudioWidgets/Private/AudioMaterialSlate/SAudioMaterialLabeledSlider.cpp
@@ -0,0 +1,230 @@
+// Copyright Epic Games, Inc. All Rights Reserved.
+
+#include "AudioMaterialSlate/SAudioMaterialLabeledSlider.h"
+#include "AudioMaterialSlate/SAudioMaterialSlider.h"
+#include "SAudioTextBox.h"
+#include "Widgets/Layout/SWidgetSwitcher.h"
+#include "Widgets/SOverlay.h"
+
+void SAudioMaterialLabeledSlider::Construct(const SAudioMaterialLabeledSlider::FArguments& InArgs)
+{
+	Style = InArgs._Style;
+	OnValueChanged = InArgs._OnValueChanged;
+	OnValueCommitted = InArgs._OnValueCommitted;
+	Orientation = InArgs._Orientation;
+	DesiredSizeOverride = InArgs._DesiredSizeOverride;
+
+	if (InArgs._SliderValue.IsSet())
+	{
+		SliderValueAttribute = InArgs._SliderValue;
+	}
+
+	// Text label
+	SAssignNew(Label, SAudioTextBox)
+		.Style(&Style->TextBoxStyle)
+		.OnValueTextCommitted_Lambda([this](const FText& Text, ETextCommit::Type CommitType)
+		{
+			const float OutputValue = FCString::Atof(*Text.ToString());
+			const float NewSliderValue = GetSliderValueForText(OutputValue);
+			if (!FMath::IsNearlyEqual(NewSliderValue, SliderValueAttribute.Get()))
+			{
+				SliderValueAttribute.Set(NewSliderValue);
+				Slider->SetValue(NewSliderValue);
+				OnValueChanged.ExecuteIfBound(NewSliderValue);
+				OnValueCommitted.ExecuteIfBound(NewSliderValue);
+			}
+		});
+
+	// Underlying slider widget
+	SAssignNew(Slider, SAudioMaterialSlider)
+		.ValueAttribute(SliderValueAttribute.Get())
+		.Owner(InArgs._Owner)
+		.Orientation(Orientation.Get())
+		.OnValueChanged_Lambda([this](float Value)
+		{
+			SliderValueAttribute.Set(Value);
+			OnValueChanged.ExecuteIfBound(Value);
+			const float OutputValue = GetOutputValueForText(Value);
+			Label->SetValueText(OutputValue);
+		});
+
+	ChildSlot
+	[
+		CreateWidgetLayout()
+	];
+}
+
+void SAudioMaterialLabeledSlider::SetSliderValue(float InSliderValue)
+{
+	SliderValueAttribute.Set(InSliderValue);
+	const float OutputValueForText = GetOutputValueForText(InSliderValue);
+	Label->SetValueText(OutputValueForText);
+	Slider->SetValue(InSliderValue);
+}
+
+void SAudioMaterialLabeledSlider::SetOrientation(EOrientation InOrientation)
+{
+	SetAttribute(Orientation, TAttribute<EOrientation>(InOrientation), EInvalidateWidgetReason::Layout);
+	Slider->SetOrientation(InOrientation);
+	LayoutWidgetSwitcher->SetActiveWidgetIndex(Orientation.Get());
+}
+
+FVector2D SAudioMaterialLabeledSlider::ComputeDesiredSize(float) const
+{
+	if (DesiredSizeOverride.Get().IsSet())
+	{
+		return DesiredSizeOverride.Get().GetValue();
+	}
+	
+	if (Style)
+	{
+		const float Width = Orientation.Get() == Orient_Vertical ? Style->DesiredSize.X : Style->DesiredSize.Y;
+		const float Heigth = Orientation.Get() == Orient_Vertical ? Style->DesiredSize.Y : Style->DesiredSize.X;
+		return FVector2D(Width, Heigth);
+	}
+
+	return FVector2D::ZeroVector;
+}
+
+void SAudioMaterialLabeledSlider::SetDesiredSizeOverride(const FVector2D Size)
+{
+	SetAttribute(DesiredSizeOverride, TAttribute<TOptional<FVector2D>>(Size), EInvalidateWidgetReason::Layout);
+}
+
+const float SAudioMaterialLabeledSlider::GetOutputValue(const float InSliderValue)
+{
+	return FMath::GetMappedRangeValueClamped(FVector2D(0.0f, 1.0f), OutputRange, InSliderValue);
+}
+
+const float SAudioMaterialLabeledSlider::GetOutputValueForText(const float InSliderValue)
+{
+	return GetOutputValue(InSliderValue);
+}
+
+const float SAudioMaterialLabeledSlider::GetSliderValueForText(const float OutputValue)
+{
+	return GetSliderValue(OutputValue);
+}
+
+const float SAudioMaterialLabeledSlider::GetSliderValue(const float OutputValue)
+{
+	return FMath::GetMappedRangeValueClamped(OutputRange, FVector2D(0.0f, 1.0f) , OutputValue);
+}
+
+void SAudioMaterialLabeledSlider::SetOutputRange(const FVector2D Range)
+{
+	OutputRange = Range;
+	// if Range.Y < Range.X, set Range.X to Range.Y
+	OutputRange.X = FMath::Min(Range.X, Range.Y);
+
+	const float OutputValue = GetOutputValue(SliderValueAttribute.Get());
+	const float ClampedOutputValue = FMath::Clamp(OutputValue, OutputRange.X, OutputRange.Y);
+	const float ClampedSliderValue = GetSliderValue(ClampedOutputValue);
+	SetSliderValue(ClampedSliderValue);
+
+	Label->UpdateValueTextWidth(OutputRange);
+}
+
+void SAudioMaterialLabeledSlider::SetLabelBackgroundColor(FSlateColor InColor)
+{
+	Label->SetLabelBackgroundColor(InColor.GetSpecifiedColor());
+}
+
+void SAudioMaterialLabeledSlider::SetUnitsText(const FText Units)
+{
+	Label->SetUnitsText(Units);
+}
+
+void SAudioMaterialLabeledSlider::SetUnitsTextReadOnly(const bool bIsReadOnly)
+{
+	Label->SetUnitsTextReadOnly(bIsReadOnly);
+}
+
+void SAudioMaterialLabeledSlider::SetValueTextReadOnly(const bool bIsReadOnly)
+{
+	Label->SetValueTextReadOnly(bIsReadOnly);
+}
+
+void SAudioMaterialLabeledSlider::SetShowLabelOnlyOnHover(const bool bShowLabelOnlyOnHover)
+{
+	Label->SetShowLabelOnlyOnHover(bShowLabelOnlyOnHover);
+}
+
+void SAudioMaterialLabeledSlider::SetShowUnitsText(const bool bShowUnitsText)
+{
+	Label->SetShowUnitsText(bShowUnitsText);
+}
+
+TSharedRef<SWidgetSwitcher> SAudioMaterialLabeledSlider::CreateWidgetLayout()
+{
+	SAssignNew(LayoutWidgetSwitcher, SWidgetSwitcher);
+	// Create overall layout
+	// Horizontal orientation
+	LayoutWidgetSwitcher->AddSlot(EOrientation::Orient_Horizontal)
+	[
+		SNew(SOverlay)
+		+ SOverlay::Slot()
+		.HAlign(HAlign_Fill)
+		.VAlign(VAlign_Center)
+		[
+			SNew(SHorizontalBox)
+			+ SHorizontalBox::Slot()
+			.HAlign(HAlign_Fill)
+			.VAlign(VAlign_Center)
+			[
+				SNew(SOverlay)
+				+ SOverlay::Slot()
+				// SSlider
+				.HAlign(HAlign_Fill)
+				.VAlign(VAlign_Center)
+				[
+					Slider.ToSharedRef()
+				]
+			]
+			// Text Label
+			+ SHorizontalBox::Slot()
+			.AutoWidth()
+			.HAlign(HAlign_Fill)
+			.VAlign(VAlign_Center)
+			[
+				Label.ToSharedRef()
+			]
+		]	
+	];
+	// Vertical orientation
+	LayoutWidgetSwitcher->AddSlot(EOrientation::Orient_Vertical)
+	[
+		SNew(SOverlay)
+		+ SOverlay::Slot()
+		.HAlign(HAlign_Center)
+		.VAlign(VAlign_Fill)
+		[
+			SNew(SVerticalBox)
+			// Text Label
+			+ SVerticalBox::Slot()
+			.AutoHeight()
+			.HAlign(HAlign_Center)
+			.VAlign(VAlign_Center)
+			[
+				Label.ToSharedRef()
+			]
+			+ SVerticalBox::Slot()
+			.HAlign(HAlign_Center)
+			.VAlign(VAlign_Fill)
+			[
+				SNew(SOverlay)
+				+ SOverlay::Slot()
+				// Actual SSlider
+				.HAlign(HAlign_Center)
+				.VAlign(VAlign_Fill)
+				[
+					Slider.ToSharedRef()
+				]
+			]
+		]		
+	];
+	LayoutWidgetSwitcher->SetActiveWidgetIndex(Orientation.Get());
+	SetOrientation(Orientation.Get());
+
+	return LayoutWidgetSwitcher.ToSharedRef();
+}
diff --git a/Engine/Plugins/Runtime/AudioWidgets/Source/AudioWidgets/Private/AudioMaterialSlate/SAudioMaterialSlider.cpp b/Engine/Plugins/Runtime/AudioWidgets/Source/AudioWidgets/Private/AudioMaterialSlate/SAudioMaterialSlider.cpp
--- a/Engine/Plugins/Runtime/AudioWidgets/Source/AudioWidgets/Private/AudioMaterialSlate/SAudioMaterialSlider.cpp
+++ b/Engine/Plugins/Runtime/AudioWidgets/Source/AudioWidgets/Private/AudioMaterialSlate/SAudioMaterialSlider.cpp
@@ -3,6 +3,7 @@
 
 #include "AudioMaterialSlate/SAudioMaterialSlider.h"
 #include "AudioMaterialSlate/AudioMaterialSlider.h"
+#include "AudioWidgetsStyle.h"
 #include "SlateOptMacros.h"
 #include "Styling/SlateBrush.h"
 
@@ -11,44 +12,80 @@ BEGIN_SLATE_FUNCTION_BUILD_OPTIMIZATION
 void SAudioMaterialSlider::Construct(const FArguments& InArgs)
 {
 	Owner = InArgs._Owner;
-	OnValueChanged = InArgs._OnValueChanged;
-	ValueAttribute = InArgs._ValueAttribute.Get();
+	Orientation = InArgs._Orientation;
 	AudioMaterialSliderStyle = InArgs._AudioMaterialSliderStyle;
+	OnValueChanged = InArgs._OnValueChanged;
+	OnValueCommitted = InArgs._OnValueCommitted;
+
+	//For now check if owner is not a widget -> use the default style.
+	if (!Cast<UWidget>(Owner))
+	{
+		AudioMaterialSliderStyle = &FAudioWidgetsStyle::Get().GetWidgetStyle<FAudioMaterialSliderStyle>("AudioMaterialSlider.Style");
+	}
 
 	ApplyNewMaterial();
-	CommitValue(InArgs._ValueAttribute.Get());
+
+	if (InArgs._ValueAttribute.IsSet())
+	{
+		CommitValue(InArgs._ValueAttribute.Get());
+	}
 }
 
 END_SLATE_FUNCTION_BUILD_OPTIMIZATION
 
 int32 SAudioMaterialSlider::OnPaint(const FPaintArgs& Args, const FGeometry& AllottedGeometry, const FSlateRect& MyCullingRect, FSlateWindowElementList& OutDrawElements, int32 LayerId, const FWidgetStyle& InWidgetStyle, bool bParentEnabled) const
 {
-
 	if (AudioMaterialSliderStyle)
 	{
-		UMaterialInstanceDynamic* DynamicMaterial = AudioMaterialSliderStyle->GetDynamicMaterial();
-		if (IsValid(DynamicMaterial))
+		if (DynamicMaterial.IsValid())
 		{
 			const float Value = ValueAttribute.Get();
-			DynamicMaterial->SetVectorParameterValue(FName("BarColor"), AudioMaterialSliderStyle->BarMainColor);
-			DynamicMaterial->SetVectorParameterValue(FName("LedColor"), AudioMaterialSliderStyle->BarShadowColor);
-			DynamicMaterial->SetVectorParameterValue(FName("BarInnerShadow"), AudioMaterialSliderStyle->BarAccentColor);
-			DynamicMaterial->SetVectorParameterValue(FName("ValueColor"), AudioMaterialSliderStyle->HandleMainColor);
-			DynamicMaterial->SetVectorParameterValue(FName("DotBevel1"), AudioMaterialSliderStyle->HandleOutlineColor);
-			DynamicMaterial->SetScalarParameterValue(FName("VALUE"), FMath::Clamp(Value, 0.f, 1.f));
-
-			DynamicMaterial->SetScalarParameterValue(FName("LocalWidth"), AllottedGeometry.GetLocalSize().X);
-			DynamicMaterial->SetScalarParameterValue(FName("LocalHeigth"), AllottedGeometry.GetLocalSize().Y);			
-		}
+			DynamicMaterial.Get()->SetVectorParameterValue(FName("BarColor"), AudioMaterialSliderStyle->BarMainColor);
+			DynamicMaterial.Get()->SetVectorParameterValue(FName("LedColor"), AudioMaterialSliderStyle->BarShadowColor);
+			DynamicMaterial.Get()->SetVectorParameterValue(FName("BarInnerShadow"), AudioMaterialSliderStyle->BarAccentColor);
+			DynamicMaterial.Get()->SetVectorParameterValue(FName("ValueColor"), AudioMaterialSliderStyle->HandleMainColor);
+			DynamicMaterial.Get()->SetVectorParameterValue(FName("DotBevel1"), AudioMaterialSliderStyle->HandleOutlineColor);
+			DynamicMaterial.Get()->SetScalarParameterValue(FName("VALUE"), FMath::Clamp(Value, 0.f, 1.f));
+			DynamicMaterial.Get()->SetScalarParameterValue(FName("LedInt"), FMath::GetMappedRangeValueClamped(FVector2D(0.f,1.f),FVector2D(0.7f,4.f), Value));
+
+			DynamicMaterial.Get()->SetScalarParameterValue(FName("LocalWidth"), AllottedGeometry.GetLocalSize().X);
+			DynamicMaterial.Get()->SetScalarParameterValue(FName("LocalHeigth"), AllottedGeometry.GetLocalSize().Y);
+
+
+			// We draw the slider as a vertical slider regardless of the orientation and apply a render transform to make it display in the correct orientation.
+			FGeometry SliderGeometry = AllottedGeometry;
+			const float AllottedWidth = Orientation == Orient_Vertical ? AllottedGeometry.GetLocalSize().X : AllottedGeometry.GetLocalSize().Y;
+			const float AllottedHeight = Orientation == Orient_Vertical ? AllottedGeometry.GetLocalSize().Y : AllottedGeometry.GetLocalSize().X;
+
+			// rotate the slider 90deg if should be horizontal.
+			if (Orientation == Orient_Horizontal)
+			{
+				// rotate
+				FSlateRenderTransform SlateRenderTransform = TransformCast<FSlateRenderTransform>(Concatenate(Inverse(FVector2f(0, AllottedHeight)), FQuat2D(FMath::DegreesToRadians(90.0f))));
+				// create a child geometry matching this one, but with the rotated render transform that will be passed to the drawed slider.
+				SliderGeometry = AllottedGeometry.MakeChild(
+					FVector2f(AllottedWidth, AllottedHeight),
+					FSlateLayoutTransform(),
+					SlateRenderTransform, FVector2f::ZeroVector);
+			}
+
+			const bool bEnabled = ShouldBeEnabled(bParentEnabled);
+			const ESlateDrawEffect DrawEffects = bEnabled ? ESlateDrawEffect::None : ESlateDrawEffect::DisabledEffect;
+				
+			const FLinearColor FinalColorAndOpacity(InWidgetStyle.GetColorAndOpacityTint());
+
+			FSlateBrush Brush;
+			Brush.SetResourceObject(DynamicMaterial.Get());
+			FSlateDrawElement::MakeBox(OutDrawElements, LayerId++, SliderGeometry.ToPaintGeometry(FVector2D(AllottedWidth, AllottedHeight), FSlateLayoutTransform()), &Brush, DrawEffects, FinalColorAndOpacity);
 
-		const bool bEnabled = ShouldBeEnabled(bParentEnabled);
-		const ESlateDrawEffect DrawEffects = bEnabled ? ESlateDrawEffect::None : ESlateDrawEffect::DisabledEffect;
-			
-		const FLinearColor FinalColorAndOpacity(InWidgetStyle.GetColorAndOpacityTint());
-
-		FSlateBrush Brush;
-		Brush.SetResourceObject(DynamicMaterial);
-		FSlateDrawElement::MakeBox(OutDrawElements, LayerId++, AllottedGeometry.ToPaintGeometry(), &Brush, DrawEffects, FinalColorAndOpacity);
+		}
+		else
+		{
+			if (AudioMaterialSliderStyle)
+			{
+				DynamicMaterial = AudioMaterialSliderStyle->CreateDynamicMaterial(Owner.Get());
+			}
+		}
 	}
 
 	return LayerId;
@@ -56,9 +93,16 @@ int32 SAudioMaterialSlider::OnPaint(const FPaintArgs& Args, const FGeometry& All
 
 FVector2D SAudioMaterialSlider::ComputeDesiredSize(float) const
 {
+	if (DesiredSizeOverride.Get().IsSet())
+	{
+		return DesiredSizeOverride.Get().GetValue();
+	}
+
 	if (AudioMaterialSliderStyle)
 	{
-		return FVector2D(AudioMaterialSliderStyle->DesiredSize);
+		const float Width = Orientation == Orient_Vertical ? AudioMaterialSliderStyle->DesiredSize.X : AudioMaterialSliderStyle->DesiredSize.Y;
+		const float Heigth = Orientation == Orient_Vertical ? AudioMaterialSliderStyle->DesiredSize.Y : AudioMaterialSliderStyle->DesiredSize.X;
+		return FVector2D(Width, Heigth);
 	}
 
 	return FVector2D::ZeroVector;
@@ -73,7 +117,16 @@ void SAudioMaterialSlider::ApplyNewMaterial()
 {
 	if (AudioMaterialSliderStyle)
 	{
-		AudioMaterialSliderStyle->CreateDynamicMaterial(Owner.Get());
+		DynamicMaterial = AudioMaterialSliderStyle->CreateDynamicMaterial(Owner.Get());
+	}
+}
+
+void SAudioMaterialSlider::SetOrientation(EOrientation InOrientation)
+{
+	if (Orientation != InOrientation)
+	{
+		Orientation = InOrientation;
+		Invalidate(EInvalidateWidgetReason::Layout);
 	}
 }
 
@@ -82,7 +135,8 @@ FReply SAudioMaterialSlider::OnMouseButtonDown(const FGeometry& MyGeometry, cons
 	if ((MouseEvent.GetEffectingButton() == EKeys::LeftMouseButton))
 	{
 		FVector2D Value2D = PositionToValue(MyGeometry, MouseEvent.GetLastScreenSpacePosition());
-		CommitValue(Value2D.Y);
+		const float Value = Orientation == Orient_Horizontal ? Value2D.X : Value2D.Y;
+		CommitValue(Value);
 
 		return FReply::Handled().CaptureMouse(SharedThis(this));
 	}
@@ -104,9 +158,9 @@ FReply SAudioMaterialSlider::OnMouseMove(const FGeometry& MyGeometry, const FPoi
 {
 	if (this->HasMouseCapture())
 	{
-		FVector2D Vector2D = PositionToValue(MyGeometry, MouseEvent.GetLastScreenSpacePosition());
-
-		CommitValue(Vector2D.Y);
+		FVector2D Value2D = PositionToValue(MyGeometry, MouseEvent.GetLastScreenSpacePosition());
+		const float Value = Orientation == Orient_Horizontal ? Value2D.X : Value2D.Y;
+		CommitValue(Value);
 
 		return FReply::Handled();
 	}
@@ -114,22 +168,33 @@ FReply SAudioMaterialSlider::OnMouseMove(const FGeometry& MyGeometry, const FPoi
 	return FReply::Unhandled();
 }
 
+void SAudioMaterialSlider::OnMouseCaptureLost(const FCaptureLostEvent& CaptureLostEvent)
+{
+	OnValueCommitted.ExecuteIfBound(ValueAttribute.Get());
+	SLeafWidget::OnMouseCaptureLost(CaptureLostEvent);
+}
 void SAudioMaterialSlider::CommitValue(float NewValue)
 {
+	const float OldValue = ValueAttribute.Get();
 	float Val = FMath::Clamp(NewValue, 0.f , 1.f);
-	if (!ValueAttribute.IsBound())
+
+	if (NewValue != OldValue)
 	{
-		ValueAttribute.Set(Val);
-	}
+		if (!ValueAttribute.IsBound())
+		{
+			ValueAttribute.Set(Val);
+		}
 
-	OnValueChanged.ExecuteIfBound(Val);
+		Invalidate(EInvalidateWidgetReason::Paint);
+		OnValueChanged.ExecuteIfBound(Val);
+	}
 }
 
 FVector2D SAudioMaterialSlider::PositionToValue(const FGeometry& MyGeometry, const FVector2D& AbsolutePosition)
 {
 	const FVector2D LocalPosition = MyGeometry.AbsoluteToLocal(AbsolutePosition);
 
-	float RelativeValueX = 1 - (LocalPosition.X) / (MyGeometry.Size.X);
+	float RelativeValueX = (LocalPosition.X) / (MyGeometry.Size.X);
 	float RelativeValueY = 1 - (LocalPosition.Y) / (MyGeometry.Size.Y);
 
 	RelativeValueX = FMath::Clamp(RelativeValueX, 0.0f, 1.0f);
diff --git a/Engine/Plugins/Runtime/AudioWidgets/Source/AudioWidgets/Private/AudioWidgetsStyle.cpp b/Engine/Plugins/Runtime/AudioWidgets/Source/AudioWidgets/Private/AudioWidgetsStyle.cpp
--- a/Engine/Plugins/Runtime/AudioWidgets/Source/AudioWidgets/Private/AudioWidgetsStyle.cpp
+++ b/Engine/Plugins/Runtime/AudioWidgets/Source/AudioWidgets/Private/AudioWidgetsStyle.cpp
@@ -161,9 +161,10 @@ FAudioWidgetsStyle::FAudioWidgetsStyle()
 	/**
 	**AudioMaterialSlider Style
 	*/
-	UMaterialInterface* SliderDefaultMaterial = LoadObject<UMaterialInterface>(nullptr, TEXT("/AudioWidgets/AudioMaterialSlate/MI_AudioMaterialSlider.MI_AudioMaterialSlider"));
+	UMaterialInterface* SliderDefaultMaterial = LoadObject<UMaterialInterface>(nullptr, TEXT("/AudioWidgets/AudioMaterialSlate/MI_AudioMaterialRoundedSlider.MI_AudioMaterialRoundedSlider"));
 
 	Set("AudioMaterialSlider.Style", FAudioMaterialSliderStyle()
+		.SetTextBoxStyle(FAudioTextBoxStyle::GetDefault())
 		.SetMaterial(SliderDefaultMaterial));
 
 	FSlateStyleRegistry::RegisterSlateStyle(*this);
diff --git a/Engine/Plugins/Runtime/AudioWidgets/Source/AudioWidgets/Public/AudioMaterialSlate/AudioMaterialSlateTypes.h b/Engine/Plugins/Runtime/AudioWidgets/Source/AudioWidgets/Public/AudioMaterialSlate/AudioMaterialSlateTypes.h
--- a/Engine/Plugins/Runtime/AudioWidgets/Source/AudioWidgets/Public/AudioMaterialSlate/AudioMaterialSlateTypes.h
+++ b/Engine/Plugins/Runtime/AudioWidgets/Source/AudioWidgets/Public/AudioMaterialSlate/AudioMaterialSlateTypes.h
@@ -2,6 +2,7 @@
 
 #pragma once
 
+#include "AudioWidgetsSlateTypes.h"
 #include "Math/Color.h"
 #include "Materials/MaterialInstanceDynamic.h"
 #include "Materials/MaterialInterface.h"
@@ -122,6 +123,11 @@ struct AUDIOWIDGETS_API FAudioMaterialSliderStyle : public FAudioMaterialWidgetS
 	FLinearColor HandleOutlineColor;
 	FAudioMaterialSliderStyle& SetSliderHandleOutlineColor(const FLinearColor& InColor) { HandleOutlineColor = InColor; return *this; }
 
+	/** The style to use for the audio text box widget. */
+	UPROPERTY(EditAnywhere, BlueprintReadWrite, Category = "Style")
+	FAudioTextBoxStyle TextBoxStyle;
+	FAudioMaterialSliderStyle& SetTextBoxStyle(const FAudioTextBoxStyle& InTextBoxStyle) { TextBoxStyle = InTextBoxStyle; return *this; }
+
 };
 
 /**
diff --git a/Engine/Plugins/Runtime/AudioWidgets/Source/AudioWidgets/Public/AudioMaterialSlate/AudioMaterialSlider.h b/Engine/Plugins/Runtime/AudioWidgets/Source/AudioWidgets/Public/AudioMaterialSlate/AudioMaterialSlider.h
--- a/Engine/Plugins/Runtime/AudioWidgets/Source/AudioWidgets/Public/AudioMaterialSlate/AudioMaterialSlider.h
+++ b/Engine/Plugins/Runtime/AudioWidgets/Source/AudioWidgets/Public/AudioMaterialSlate/AudioMaterialSlider.h
@@ -72,6 +72,10 @@ private:
 	UPROPERTY(EditAnywhere, BlueprintSetter = SetValue, BlueprintGetter = GetValue, Category = "Appearance", meta = (UIMin = "0", UIMax = "1"))
 	float Value = 1.f;
 
+	/**Orientation of the slider*/
+	UPROPERTY(EditAnywhere, Category = "Appearance")
+	TEnumAsByte<EOrientation> Orientation = EOrientation::Orient_Horizontal;
+
 private:
 
 	/** Native Slate Widget */
diff --git a/Engine/Plugins/Runtime/AudioWidgets/Source/AudioWidgets/Public/AudioMaterialSlate/SAudioMaterialLabeledSlider.h b/Engine/Plugins/Runtime/AudioWidgets/Source/AudioWidgets/Public/AudioMaterialSlate/SAudioMaterialLabeledSlider.h
new file mode 100644
--- /dev/null
+++ b/Engine/Plugins/Runtime/AudioWidgets/Source/AudioWidgets/Public/AudioMaterialSlate/SAudioMaterialLabeledSlider.h
@@ -0,0 +1,115 @@
+// Copyright Epic Games, Inc. All Rights Reserved.
+
+#pragma once
+
+#include "AudioMaterialSlate/AudioMaterialSlateTypes.h"
+#include "Framework/SlateDelegates.h"
+#include "SAudioInputWidget.h"
+#include "Styling/ISlateStyle.h"
+#include "Styling/SlateStyleRegistry.h"
+#include "Styling/SlateWidgetStyleAsset.h"
+
+class SAudioTextBox;
+class SAudioMaterialSlider;
+class SWidgetSwitcher;
+
+/**
+ * Wraps SAudioMaterialSlider and adds Label text that will show a value text.
+ */
+class AUDIOWIDGETS_API SAudioMaterialLabeledSlider
+	: public SAudioInputWidget
+{
+public:
+	SLATE_BEGIN_ARGS(SAudioMaterialLabeledSlider)
+	{}
+
+	/** The style used to draw the slider. */
+	SLATE_STYLE_ARGUMENT(FAudioMaterialSliderStyle, Style)
+
+	/** The owner object*/
+	SLATE_ARGUMENT(TWeakObjectPtr<UObject>, Owner)
+
+	/** A value representing the normalized linear (0 - 1) slider value position. */
+	SLATE_ATTRIBUTE(float, SliderValue)
+		
+	/** The slider's orientation. */
+	SLATE_ARGUMENT(EOrientation, Orientation)
+	
+	/** When specified, use this as the slider's desired size */
+	SLATE_ATTRIBUTE(TOptional<FVector2D>, DesiredSizeOverride)
+
+	/** Called when the value is changed by slider or typing */
+	SLATE_EVENT(FOnFloatValueChanged, OnValueChanged)
+
+	/** Called when the value is committed by mouse capture ending */
+	SLATE_EVENT(FOnFloatValueChanged, OnValueCommitted)
+
+	SLATE_END_ARGS()
+
+	// Holds a delegate that is executed when the slider's value changed.
+	FOnFloatValueChanged OnValueChanged;
+
+	// Holds a delegate that is executed when the slider's value is committed (mouse capture ends).
+	FOnFloatValueChanged OnValueCommitted;
+
+	/**
+	 * Construct the widget.
+	 */
+	virtual void Construct(const SAudioMaterialLabeledSlider::FArguments& InArgs);
+	virtual const float GetOutputValue(const float InSliderValue) override;
+	virtual const float GetSliderValue(const float OutputValue) override;
+	virtual const float GetOutputValueForText(const float InSliderValue);
+	virtual const float GetSliderValueForText(const float OutputValue);
+	
+	//SAudioInputWidget
+	/**
+	 * Set the slider's linear (0-1 normalized) value. 
+	 */
+	virtual void SetSliderValue(float InSliderValue) override;
+	FVector2D ComputeDesiredSize(float) const;
+	virtual void SetDesiredSizeOverride(const FVector2D DesiredSize) override;
+
+	void SetOrientation(EOrientation InOrientation);
+	virtual void SetOutputRange(const FVector2D Range) override;
+
+	// Text label functions 
+	void SetLabelBackgroundColor(FSlateColor InColor) override;
+	void SetUnitsText(const FText Units) override;
+	void SetUnitsTextReadOnly(const bool bIsReadOnly) override;
+	void SetValueTextReadOnly(const bool bIsReadOnly);
+	void SetShowLabelOnlyOnHover(const bool bShowLabelOnlyOnHover);
+	void SetShowUnitsText(const bool bShowUnitsText) override;
+	//~SAudioInputWidget
+
+private:
+
+	// Holds the style for the Slider
+	const FAudioMaterialSliderStyle* Style;
+	
+	// Holds the slider's current linear value
+	TAttribute<float> SliderValueAttribute;
+
+	// Holds the slider's orientation
+	TAttribute<EOrientation> Orientation;
+
+	// Optional override for desired size 
+	TAttribute<TOptional<FVector2D>> DesiredSizeOverride;
+
+	// Label text bg color 
+	TAttribute<FSlateColor> LabelBackgroundColor;
+
+	// Widget components
+	TSharedPtr<SAudioMaterialSlider> Slider;
+	TSharedPtr<SAudioTextBox> Label;
+
+	// Range for output
+	FVector2D OutputRange = FVector2D(0.0f, 1.0f);
+	static const FVector2D NormalizedLinearSliderRange;
+
+	/** Switches between the vertical and horizontal views */
+	TSharedPtr<SWidgetSwitcher> LayoutWidgetSwitcher;
+
+private:
+
+	TSharedRef<SWidgetSwitcher> CreateWidgetLayout();
+};
diff --git a/Engine/Plugins/Runtime/AudioWidgets/Source/AudioWidgets/Public/AudioMaterialSlate/SAudioMaterialSlider.h b/Engine/Plugins/Runtime/AudioWidgets/Source/AudioWidgets/Public/AudioMaterialSlate/SAudioMaterialSlider.h
--- a/Engine/Plugins/Runtime/AudioWidgets/Source/AudioWidgets/Public/AudioMaterialSlate/SAudioMaterialSlider.h
+++ b/Engine/Plugins/Runtime/AudioWidgets/Source/AudioWidgets/Public/AudioMaterialSlate/SAudioMaterialSlider.h
@@ -23,6 +23,9 @@ public:
 	/** The owner object*/
 	SLATE_ARGUMENT(TWeakObjectPtr<UObject>, Owner)
 
+	/** The slider's orientation. */
+	SLATE_ARGUMENT(EOrientation, Orientation)
+
 	/** The style used to draw the slider. */
 	SLATE_STYLE_ARGUMENT(FAudioMaterialSliderStyle, AudioMaterialSliderStyle)
 
@@ -31,6 +34,9 @@ public:
 
 	/** Called when the value is changed by the slider. */
 	SLATE_EVENT(FOnFloatValueChanged, OnValueChanged)
+	
+	/** Called when the value is committed (mouse capture ends) */
+	SLATE_EVENT(FOnFloatValueChanged, OnValueCommitted)
 
 	SLATE_END_ARGS()
 
@@ -49,15 +55,23 @@ public:
 	/** Apply new material to be used to render the Slate.*/
 	void ApplyNewMaterial();
 
-	// SWidget overrides
+	/** Set the orientation of the slider*/
+	void SetOrientation(EOrientation InOrientation);
+
+	//SWidget
 	virtual FReply OnMouseButtonDown(const FGeometry& MyGeometry, const FPointerEvent& MouseEvent) override;
 	virtual FReply OnMouseButtonUp(const FGeometry& MyGeometry, const FPointerEvent& MouseEvent) override;
 	virtual FReply OnMouseMove(const FGeometry& MyGeometry, const FPointerEvent& MouseEvent) override;
+	virtual void OnMouseCaptureLost(const FCaptureLostEvent& CaptureLostEvent) override;
+	//~SWidget
 
 public:
 
 	// Holds a delegate that is executed when the slider's value changed.
-	FOnFloatValueChanged OnValueChanged;	
+	FOnFloatValueChanged OnValueChanged;
+	
+	// Holds a delegate that is executed when the slider's value is committed (mouse capture ends).
+	FOnFloatValueChanged OnValueCommitted;
 
 private:
 
@@ -79,12 +93,21 @@ private:
 
 private:
 
+	// Optional override for desired size 
+	TAttribute<TOptional<FVector2D>> DesiredSizeOverride;
+
+	// Holds the slider's orientation
+	EOrientation Orientation;
+
 	// Holds the owner of the Slate
 	TWeakObjectPtr<UObject> Owner;
 
 	// Holds the style for the Slate
 	const FAudioMaterialSliderStyle* AudioMaterialSliderStyle = nullptr;
 
+	// Holds the Modifiable Material that represent the slider
+	mutable TWeakObjectPtr<UMaterialInstanceDynamic> DynamicMaterial;
+
 	//Holds the current value
 	TAttribute<float> ValueAttribute = 0.f;
 
diff --git a/Engine/Plugins/Runtime/Metasound/Source/MetasoundEditor/Private/SMetasoundGraphNode.cpp b/Engine/Plugins/Runtime/Metasound/Source/MetasoundEditor/Private/SMetasoundGraphNode.cpp
--- a/Engine/Plugins/Runtime/Metasound/Source/MetasoundEditor/Private/SMetasoundGraphNode.cpp
+++ b/Engine/Plugins/Runtime/Metasound/Source/MetasoundEditor/Private/SMetasoundGraphNode.cpp
@@ -2,6 +2,7 @@
 #include "SMetasoundGraphNode.h"
 
 #include "AudioMaterialSlate/SAudioMaterialKnob.h"
+#include "AudioMaterialSlate/SAudioMaterialLabeledSlider.h"
 #include "AudioParameterControllerInterface.h"
 #include "Components/AudioComponent.h"
 #include "GraphEditorSettings.h"
@@ -730,26 +731,36 @@ namespace Metasound
 
 						if (DefaultFloat->WidgetType == EMetasoundMemberDefaultWidget::Slider)
 						{
-							// Create slider 
-							if (DefaultFloat->WidgetValueType == EMetasoundMemberDefaultWidgetValueType::Frequency)
-							{
-								SAssignNew(InputWidget, SAudioFrequencySlider)
-									.OnValueChanged_Lambda(OnValueChangedLambda)
-									.OnValueCommitted_Lambda(OnValueCommittedLambda);
-							}
-							else if (DefaultFloat->WidgetValueType == EMetasoundMemberDefaultWidgetValueType::Volume)
+							if (Metasound::Editor::GraphNodePrivate::UseAudioMaterialWidgets)
 							{
-								SAssignNew(InputWidget, SAudioVolumeSlider)
+								SAssignNew(InputWidget, SAudioMaterialLabeledSlider)
+									.Owner(GraphMember->GetOwningGraph())
 									.OnValueChanged_Lambda(OnValueChangedLambda)
 									.OnValueCommitted_Lambda(OnValueCommittedLambda);
-								StaticCastSharedPtr<SAudioVolumeSlider>(InputWidget)->SetUseLinearOutput(DefaultFloat->VolumeWidgetUseLinearOutput);
 							}
 							else
 							{
-								SAssignNew(InputWidget, SAudioSlider)
-									.OnValueChanged_Lambda(OnValueChangedLambda)
-									.OnValueCommitted_Lambda(OnValueCommittedLambda);
-								InputWidget->SetShowUnitsText(false);
+								// Create slider 
+								if (DefaultFloat->WidgetValueType == EMetasoundMemberDefaultWidgetValueType::Frequency)
+								{
+									SAssignNew(InputWidget, SAudioFrequencySlider)
+										.OnValueChanged_Lambda(OnValueChangedLambda)
+										.OnValueCommitted_Lambda(OnValueCommittedLambda);
+								}
+								else if (DefaultFloat->WidgetValueType == EMetasoundMemberDefaultWidgetValueType::Volume)
+								{
+									SAssignNew(InputWidget, SAudioVolumeSlider)
+										.OnValueChanged_Lambda(OnValueChangedLambda)
+										.OnValueCommitted_Lambda(OnValueCommittedLambda);
+									StaticCastSharedPtr<SAudioVolumeSlider>(InputWidget)->SetUseLinearOutput(DefaultFloat->VolumeWidgetUseLinearOutput);
+								}
+								else
+								{
+									SAssignNew(InputWidget, SAudioSlider)
+										.OnValueChanged_Lambda(OnValueChangedLambda)
+										.OnValueCommitted_Lambda(OnValueCommittedLambda);
+									InputWidget->SetShowUnitsText(false);
+								}
 							}
 							// Slider layout 
 							if (DefaultFloat->WidgetOrientation == Orient_Vertical)
@@ -804,10 +815,19 @@ namespace Metasound
 									[
 										Slot2.ToSharedRef()
 									];
-								InputWidget->SetDesiredSizeOverride(FVector2D(SliderDesiredSizeVertical.Y, SliderDesiredSizeVertical.X));
-							}
-							// safe downcast because the ptr was just assigned above 
-							StaticCastSharedPtr<SAudioSliderBase>(InputWidget)->SetOrientation(DefaultFloat->WidgetOrientation);
+									InputWidget->SetDesiredSizeOverride(FVector2D(SliderDesiredSizeVertical.Y, SliderDesiredSizeVertical.X));
+								}
+
+								if (Metasound::Editor::GraphNodePrivate::UseAudioMaterialWidgets)
+								{
+									// safe downcast because the ptr was just assigned above 
+									StaticCastSharedPtr<SAudioMaterialLabeledSlider>(InputWidget)->SetOrientation(DefaultFloat->WidgetOrientation);
+								}
+								else
+								{
+									// safe downcast because the ptr was just assigned above 
+									StaticCastSharedPtr<SAudioSliderBase>(InputWidget)->SetOrientation(DefaultFloat->WidgetOrientation);
+								}
 						}
 						else if (DefaultFloat->WidgetType == EMetasoundMemberDefaultWidget::RadialSlider)
 						{
