Value text label to Metasound AudioMaterialKnob widget
#rnx

diff --git a/Engine/Plugins/Runtime/AudioWidgets/Source/AudioWidgets/Private/AudioMaterialSlate/SAudioMaterialKnob.cpp b/Engine/Plugins/Runtime/AudioWidgets/Source/AudioWidgets/Private/AudioMaterialSlate/SAudioMaterialKnob.cpp
--- a/Engine/Plugins/Runtime/AudioWidgets/Source/AudioWidgets/Private/AudioMaterialSlate/SAudioMaterialKnob.cpp
+++ b/Engine/Plugins/Runtime/AudioWidgets/Source/AudioWidgets/Private/AudioMaterialSlate/SAudioMaterialKnob.cpp
@@ -55,11 +55,6 @@ const float SAudioMaterialKnob::GetSliderValue(const float OutputValue)
 	return FMath::GetMappedRangeValueClamped(OutputRange, NormalizedLinearSliderRange, OutputValue);
 }
 
-void SAudioMaterialKnob::SetSliderValue(float InSliderValue)
-{
-	ValueAttribute.Set(InSliderValue);
-}
-
 void SAudioMaterialKnob::SetOutputRange(const FVector2D Range)
 {
 	OutputRange = Range;
@@ -69,7 +64,7 @@ void SAudioMaterialKnob::SetOutputRange(const FVector2D Range)
 	const float OutputValue = GetOutputValue(ValueAttribute.Get());
 	const float ClampedOutputValue = FMath::Clamp(OutputValue, OutputRange.X, OutputRange.Y);
 	const float ClampedSliderValue = GetSliderValue(ClampedOutputValue);
-	SetSliderValue(ClampedSliderValue);
+	ValueAttribute.Set(ClampedSliderValue);
 }
 
 void SAudioMaterialKnob::SetDesiredSizeOverride(const FVector2D Size)
@@ -199,7 +194,17 @@ FReply SAudioMaterialKnob::OnMouseButtonUp(const FGeometry& MyGeometry, const FP
 
 void SAudioMaterialKnob::CommitValue(float NewValue)
 {
-	ValueAttribute.Set(NewValue);
-	Invalidate(EInvalidateWidgetReason::Paint);
-	OnValueChanged.ExecuteIfBound(NewValue);
+	const float OldValue = ValueAttribute.Get();
+	float Val = FMath::Clamp(NewValue, 0.f, 1.f);
+
+	if (NewValue != OldValue)
+	{
+		if (!ValueAttribute.IsBound())
+		{
+			ValueAttribute.Set(Val);
+		}
+
+		Invalidate(EInvalidateWidgetReason::Paint);
+		OnValueChanged.ExecuteIfBound(Val);
+	}
 }
\ No newline at end of file
diff --git a/Engine/Plugins/Runtime/AudioWidgets/Source/AudioWidgets/Private/AudioMaterialSlate/SAudioMaterialLabeledKnob.cpp b/Engine/Plugins/Runtime/AudioWidgets/Source/AudioWidgets/Private/AudioMaterialSlate/SAudioMaterialLabeledKnob.cpp
new file mode 100644
--- /dev/null
+++ b/Engine/Plugins/Runtime/AudioWidgets/Source/AudioWidgets/Private/AudioMaterialSlate/SAudioMaterialLabeledKnob.cpp
@@ -0,0 +1,164 @@
+// Copyright Epic Games, Inc. All Rights Reserved.
+
+
+#include "AudioMaterialSlate/SAudioMaterialLabeledKnob.h"
+#include "AudioMaterialSlate/SAudioMaterialKnob.h"
+#include "AudioMaterialSlate/AudioMaterialSlateTypes.h"
+#include "SAudioTextBox.h"
+#include "SlateOptMacros.h"
+#include "Widgets/SBoxPanel.h"
+
+BEGIN_SLATE_FUNCTION_BUILD_OPTIMIZATION
+void SAudioMaterialLabeledKnob::Construct(const FArguments& InArgs)
+{
+	Style = InArgs._Style;
+	Owner = InArgs._Owner;
+	OnValueChanged = InArgs._OnValueChanged;
+	OnMouseCaptureEnd = InArgs._OnMouseCaptureEnd;
+	OnMouseCaptureBegin = InArgs._OnMouseCaptureBegin;
+
+	if (InArgs._Value.IsSet())
+	{
+		ValueAttribute = InArgs._Value;
+	}
+
+	// Text label
+	SAssignNew(Label, SAudioTextBox)
+		.Style(&Style->TextBoxStyle)
+		.OnValueTextCommitted_Lambda([this](const FText& Text, ETextCommit::Type CommitType)
+		{
+			const float OutputValue = FCString::Atof(*Text.ToString());
+			const float NewSliderValue = GetSliderValueForText(OutputValue);
+			if (!FMath::IsNearlyEqual(NewSliderValue, ValueAttribute.Get()))
+			{
+				ValueAttribute.Set(NewSliderValue);
+				Knob->SetValue(NewSliderValue);
+				OnValueChanged.ExecuteIfBound(NewSliderValue);
+				OnMouseCaptureEnd.ExecuteIfBound();
+			}
+		});
+
+	// Underlying Knob widget
+	SAssignNew(Knob, SAudioMaterialKnob)
+		.Owner(InArgs._Owner)
+		.OnFloatValueChanged_Lambda([this](float Value)
+		{
+			ValueAttribute.Set(Value);
+			OnValueChanged.ExecuteIfBound(Value);
+			const float OutputValue = GetOutputValueForText(Value);
+			Label->SetValueText(OutputValue);
+		})
+		.OnMouseCaptureBegin(OnMouseCaptureBegin)
+		.OnMouseCaptureEnd(OnMouseCaptureEnd);
+
+	//Create layout
+	SAssignNew(VerticalLayotWidget, SVerticalBox)
+		+ SVerticalBox::Slot()
+		.HAlign(HAlign_Fill)
+		.VAlign(VAlign_Fill)
+		[
+			Knob.ToSharedRef()
+		]
+		+ SVerticalBox::Slot()
+		.HAlign(HAlign_Center)
+		.VAlign(VAlign_Top)
+		.AutoHeight()
+		[
+			Label.ToSharedRef()
+		];
+
+	ChildSlot
+	[
+		VerticalLayotWidget.ToSharedRef()
+	];
+
+	SetOutputRange(OutputRange);
+}
+END_SLATE_FUNCTION_BUILD_OPTIMIZATION
+
+void SAudioMaterialLabeledKnob::SetValue(float InValue)
+{
+	ValueAttribute.Set(InValue);
+	const float OutputValueForText = GetOutputValueForText(InValue);
+	Label->SetValueText(OutputValueForText);
+	Knob->SetValue(InValue);
+}
+
+const float SAudioMaterialLabeledKnob::GetOutputValue(const float InSliderValue)
+{
+	return FMath::GetMappedRangeValueClamped(NormalizedLinearSliderRange, OutputRange, InSliderValue);
+}
+
+const float SAudioMaterialLabeledKnob::GetSliderValue(const float OutputValue)
+{
+	return FMath::GetMappedRangeValueClamped(OutputRange, NormalizedLinearSliderRange, OutputValue);
+}
+
+const float SAudioMaterialLabeledKnob::GetOutputValueForText(const float InSliderValue)
+{
+	return GetOutputValue(InSliderValue);
+}
+
+const float SAudioMaterialLabeledKnob::GetSliderValueForText(const float OutputValue)
+{
+	return GetSliderValue(OutputValue);
+}
+
+void SAudioMaterialLabeledKnob::SetSliderValue(float InSliderValue)
+{
+	SetValue(InSliderValue);
+}
+
+void SAudioMaterialLabeledKnob::SetOutputRange(const FVector2D Range)
+{
+	OutputRange = Range;
+	// if Range.Y < Range.X, set Range.X to Range.Y
+	OutputRange.X = FMath::Min(Range.X, Range.Y);
+
+	const float OutputValue = GetOutputValue(ValueAttribute.Get());
+	const float ClampedOutputValue = FMath::Clamp(OutputValue, OutputRange.X, OutputRange.Y);
+	const float ClampedSliderValue = GetSliderValue(ClampedOutputValue);
+	SetSliderValue(ClampedSliderValue);
+
+	Label->UpdateValueTextWidth(OutputRange);
+}
+
+FVector2D SAudioMaterialLabeledKnob::ComputeDesiredSize(float) const
+{
+	if (DesiredSizeOverride.Get().IsSet())
+	{
+		return DesiredSizeOverride.Get().GetValue();
+	}
+	
+	if (Style)
+	{
+		return FVector2D(Style->DesiredSize);
+	}
+
+	return FVector2D::ZeroVector;
+}
+
+void SAudioMaterialLabeledKnob::SetDesiredSizeOverride(const FVector2D Size)
+{
+	SetAttribute(DesiredSizeOverride, TAttribute<TOptional<FVector2D>>(Size), EInvalidateWidgetReason::Layout);
+}
+
+void SAudioMaterialLabeledKnob::SetLabelBackgroundColor(FSlateColor InColor)
+{
+	Label->SetLabelBackgroundColor(InColor);
+}
+
+void SAudioMaterialLabeledKnob::SetUnitsText(const FText Units)
+{
+	Label->SetUnitsText(Units);
+}
+
+void SAudioMaterialLabeledKnob::SetUnitsTextReadOnly(const bool bIsReadOnly)
+{
+	Label->SetValueTextReadOnly(bIsReadOnly);
+}
+
+void SAudioMaterialLabeledKnob::SetShowUnitsText(const bool bShowUnitsText)
+{
+	Label->SetShowUnitsText(bShowUnitsText);
+}
diff --git a/Engine/Plugins/Runtime/AudioWidgets/Source/AudioWidgets/Public/AudioMaterialSlate/AudioMaterialSlateTypes.h b/Engine/Plugins/Runtime/AudioWidgets/Source/AudioWidgets/Public/AudioMaterialSlate/AudioMaterialSlateTypes.h
--- a/Engine/Plugins/Runtime/AudioWidgets/Source/AudioWidgets/Public/AudioMaterialSlate/AudioMaterialSlateTypes.h
+++ b/Engine/Plugins/Runtime/AudioWidgets/Source/AudioWidgets/Public/AudioMaterialSlate/AudioMaterialSlateTypes.h
@@ -184,6 +184,10 @@ struct AUDIOWIDGETS_API FAudioMaterialKnobStyle : public FAudioMaterialWidgetSty
 	FLinearColor KnobBarFillTintColor;
 	FAudioMaterialKnobStyle& SetKnobBarFillTintColor(const FLinearColor& InColor) { KnobBarFillTintColor = InColor; return *this; }
 
+	/** The style to use for the audio text box widget. */
+	UPROPERTY(EditAnywhere, BlueprintReadWrite, Category = "Style")
+	FAudioTextBoxStyle TextBoxStyle;
+	void SetTextBoxStyle(const FAudioTextBoxStyle& InTextBoxStyle) { TextBoxStyle = InTextBoxStyle; }
 };
 
 /**
diff --git a/Engine/Plugins/Runtime/AudioWidgets/Source/AudioWidgets/Public/AudioMaterialSlate/SAudioMaterialKnob.h b/Engine/Plugins/Runtime/AudioWidgets/Source/AudioWidgets/Public/AudioMaterialSlate/SAudioMaterialKnob.h
--- a/Engine/Plugins/Runtime/AudioWidgets/Source/AudioWidgets/Public/AudioMaterialSlate/SAudioMaterialKnob.h
+++ b/Engine/Plugins/Runtime/AudioWidgets/Source/AudioWidgets/Public/AudioMaterialSlate/SAudioMaterialKnob.h
@@ -4,7 +4,6 @@
 
 #include "AudioMaterialSlate/AudioMaterialSlateTypes.h"
 #include "Framework/SlateDelegates.h"
-#include "SAudioInputWidget.h"
 #include "Styling/ISlateStyle.h"
 #include "Styling/SlateWidgetStyleAsset.h"
 #include "Widgets/SLeafWidget.h"
@@ -15,7 +14,7 @@ class UObject;
  * A simple slate that renders a knob in single material and modifies the material on value change.
  *
  */
-class AUDIOWIDGETS_API SAudioMaterialKnob : public SAudioInputWidget
+class AUDIOWIDGETS_API SAudioMaterialKnob : public SLeafWidget
 {
 public:
 	SLATE_BEGIN_ARGS(SAudioMaterialKnob)
@@ -50,24 +49,14 @@ public:
 	/** Apply new material to be used to render the Slate.*/
 	UMaterialInstanceDynamic* ApplyNewMaterial();
 
-	//SAudioInputWidget
-	virtual const float GetOutputValue(const float InSliderValue) override;
-	virtual const float GetSliderValue(const float OutputValue) override;
+	const float GetOutputValue(const float InSliderValue);
+	const float GetSliderValue(const float OutputValue);
 
-	/**
-	 * Set the knob's linear (0-1 normalized) value.
-	 */
-	virtual void SetSliderValue(float InSliderValue) override;
-	virtual void SetOutputRange(const FVector2D Range) override;
+	/** Set the output range of the Knob*/
+	void SetOutputRange(const FVector2D Range);
 
-	virtual void SetDesiredSizeOverride(const FVector2D Size) override;
-
-	//These are pure virtual functions in the parent class and are implemented properly in this class later.
-	virtual void SetLabelBackgroundColor(FSlateColor InColor) override {};
-	virtual void SetUnitsText(const FText Units) override {};
-	virtual void SetUnitsTextReadOnly(const bool bIsReadOnly) override {};
-	virtual void SetShowUnitsText(const bool bShowUnitsText) override {};
-	//~SAudioInputWidget
+	/**Set desired size of the Slate*/
+	void SetDesiredSizeOverride(const FVector2D Size);
 
 public:
 
@@ -98,6 +87,7 @@ protected:
 
 private:
 
+	// Holds the optional style for the Slate
 	TAttribute<TOptional<FVector2D>> DesiredSizeOverride;
 
 	/**Commits new value*/
diff --git a/Engine/Plugins/Runtime/AudioWidgets/Source/AudioWidgets/Public/AudioMaterialSlate/SAudioMaterialLabeledKnob.h b/Engine/Plugins/Runtime/AudioWidgets/Source/AudioWidgets/Public/AudioMaterialSlate/SAudioMaterialLabeledKnob.h
new file mode 100644
--- /dev/null
+++ b/Engine/Plugins/Runtime/AudioWidgets/Source/AudioWidgets/Public/AudioMaterialSlate/SAudioMaterialLabeledKnob.h
@@ -0,0 +1,111 @@
+// Copyright Epic Games, Inc. All Rights Reserved.
+
+#pragma once
+
+#include "AudioMaterialSlate/AudioMaterialSlateTypes.h"
+#include "Framework/SlateDelegates.h"
+#include "SAudioInputWidget.h"
+#include "Styling/ISlateStyle.h"
+#include "Styling/SlateWidgetStyleAsset.h"
+#include "Widgets/SLeafWidget.h"
+
+class SAudioTextBox;
+class SAudioMaterialKnob;
+class SVerticalBox;
+class UObject;
+
+/**
+ * Wraps SAudioMaterialKnob and adds Label text that will show a value text.
+ */
+class AUDIOWIDGETS_API SAudioMaterialLabeledKnob : public SAudioInputWidget
+{
+public:
+	SLATE_BEGIN_ARGS(SAudioMaterialLabeledKnob)
+	{}
+
+	/** The owner object*/
+	SLATE_ARGUMENT(TWeakObjectPtr<UObject>, Owner)
+
+	/**Value of the Knob*/
+	SLATE_ATTRIBUTE(float, Value)
+
+	/** The style used to draw the knob. */
+	SLATE_STYLE_ARGUMENT(FAudioMaterialKnobStyle, Style)
+
+	/** Called when the knob's state changes. */
+	SLATE_EVENT(FOnFloatValueChanged, OnValueChanged)
+
+	/** Invoked when the mouse is pressed and a capture begins. */
+	SLATE_EVENT(FSimpleDelegate, OnMouseCaptureBegin)
+
+	/** Invoked when the mouse is released and a capture ends. */
+	SLATE_EVENT(FSimpleDelegate, OnMouseCaptureEnd)
+
+	SLATE_END_ARGS()
+
+	/** Constructs this widget with InArgs */
+	void Construct(const FArguments& InArgs);
+
+	/** Set the Value attribute */
+	void SetValue(float InValueAttribute);
+
+	//SAudioInputWidget
+	virtual const float GetOutputValue(const float InSliderValue) override;
+	virtual const float GetSliderValue(const float OutputValue) override;
+	virtual const float GetOutputValueForText(const float InSliderValue);
+	virtual const float GetSliderValueForText(const float OutputValue);
+
+	/**
+	 * Set the knob's linear (0-1 normalized) value.
+	 */
+	virtual void SetSliderValue(float InSliderValue) override;
+	virtual void SetOutputRange(const FVector2D Range) override;
+	virtual void SetDesiredSizeOverride(const FVector2D Size) override;
+	virtual void SetLabelBackgroundColor(FSlateColor InColor) override;
+	virtual void SetUnitsText(const FText Units) override;
+	virtual void SetUnitsTextReadOnly(const bool bIsReadOnly) override;
+	virtual void SetShowUnitsText(const bool bShowUnitsText) override;
+	//~SAudioInputWidget
+
+public:
+
+	// Holds a delegate that is executed when the knob's value changes.
+	FOnFloatValueChanged OnValueChanged;
+
+	// Holds a delegate that is executed when the mouse is pressed and a capture begins.
+	FSimpleDelegate OnMouseCaptureBegin;
+
+	// Holds a delegate that is executed when the mouse is let up and a capture ends.
+	FSimpleDelegate OnMouseCaptureEnd;
+
+protected:
+
+	//SWidget
+	virtual FVector2D ComputeDesiredSize(float) const override;
+	//~SWidget
+
+private:
+
+	TAttribute<TOptional<FVector2D>> DesiredSizeOverride;
+
+	// Holds the owner of the Slate
+	TWeakObjectPtr<UObject> Owner;
+
+	// Holds the style for the Slate
+	const FAudioMaterialKnobStyle* Style = nullptr;
+
+	//Holds the knobs current Value
+	TAttribute<float> ValueAttribute = 1.0f;
+
+	// Widget components
+	TSharedPtr<SAudioMaterialKnob> Knob;
+	TSharedPtr<SAudioTextBox> Label;
+
+	/** verticalBox that holds the widgets*/
+	TSharedPtr<SVerticalBox> VerticalLayotWidget;
+
+	// Range for output 
+	FVector2D OutputRange = FVector2D(0.0f, 1.0f);
+	const FVector2D NormalizedLinearSliderRange = FVector2D(0.0f, 1.0f);
+
+};
diff --git a/Engine/Plugins/Runtime/Metasound/Source/MetasoundEditor/Private/SMetasoundGraphNode.cpp b/Engine/Plugins/Runtime/Metasound/Source/MetasoundEditor/Private/SMetasoundGraphNode.cpp
--- a/Engine/Plugins/Runtime/Metasound/Source/MetasoundEditor/Private/SMetasoundGraphNode.cpp
+++ b/Engine/Plugins/Runtime/Metasound/Source/MetasoundEditor/Private/SMetasoundGraphNode.cpp
@@ -1,7 +1,7 @@
 // Copyright Epic Games, Inc. All Rights Reserved.
 #include "SMetasoundGraphNode.h"
 
-#include "AudioMaterialSlate/SAudioMaterialKnob.h"
+#include "AudioMaterialSlate/SAudioMaterialLabeledKnob.h"
 #include "AudioMaterialSlate/SAudioMaterialLabeledSlider.h"
 #include "AudioParameterControllerInterface.h"
 #include "Components/AudioComponent.h"
@@ -737,6 +737,7 @@ namespace Metasound
 									.Owner(GraphMember->GetOwningGraph())
 									.OnValueChanged_Lambda(OnValueChangedLambda)
 									.OnValueCommitted_Lambda(OnValueCommittedLambda);
+								InputWidget->SetShowUnitsText(false);
 							}
 							else
 							{
@@ -855,11 +856,12 @@ namespace Metasound
 
 							if (Metasound::Editor::GraphNodePrivate::UseAudioMaterialWidgets)
 							{
-								SAssignNew(InputWidget, SAudioMaterialKnob)
+								SAssignNew(InputWidget, SAudioMaterialLabeledKnob)
 									.Owner(GraphMember->GetOwningGraph())
-									.OnFloatValueChanged_Lambda(OnValueChangedLambda)
+									.OnValueChanged_Lambda(OnValueChangedLambda)
 									.OnMouseCaptureBegin_Lambda(OnRadialSliderMouseCaptureBeginLambda)
 									.OnMouseCaptureEnd_Lambda(OnRadialSliderMouseCaptureEndLambda);
+								InputWidget->SetShowUnitsText(false);								
 							}
 							else
 							{
@@ -905,10 +907,7 @@ namespace Metasound
 								[
 									InputWidget.ToSharedRef()
 								];
-							if (!Metasound::Editor::GraphNodePrivate::UseAudioMaterialWidgets)
-							{
 								InputWidget->SetDesiredSizeOverride(RadialSliderDesiredSize);
-							}
 						}
 
 						InputWidget->SetOutputRange(DefaultFloat->GetRange());
