Extra features to the AudioMaterialKnob
#rnx

diff --git a/Engine/Plugins/Runtime/AudioWidgets/Source/AudioWidgets/Private/AudioMaterialSlate/AudioMaterialKnob.cpp b/Engine/Plugins/Runtime/AudioWidgets/Source/AudioWidgets/Private/AudioMaterialSlate/AudioMaterialKnob.cpp
--- a/Engine/Plugins/Runtime/AudioWidgets/Source/AudioWidgets/Private/AudioMaterialSlate/AudioMaterialKnob.cpp
+++ b/Engine/Plugins/Runtime/AudioWidgets/Source/AudioWidgets/Private/AudioMaterialSlate/AudioMaterialKnob.cpp
@@ -7,10 +7,16 @@
 #include "AudioWidgetsStyle.h"
 
 #define LOCTEXT_NAMESPACE "AudioWidgets"
-UAudioMaterialKnob::UAudioMaterialKnob()
+UAudioMaterialKnob::UAudioMaterialKnob(const FObjectInitializer& ObjectInitializer)
+	: Super(ObjectInitializer)
+	, TuneSpeed(0.2f)
+	, FineTuneSpeed(0.05f)
+	, bLocked(false)
+	, bMouseUsesStep(false)
+	, StepSize(0.01f)
 {
 	//get default style
-	WidgetStyle = FAudioWidgetsStyle::Get().GetWidgetStyle<FAudioMaterialKnobStyle>("AudioMaterialKnob.Style");	
+	WidgetStyle = FAudioWidgetsStyle::Get().GetWidgetStyle<FAudioMaterialKnobStyle>("AudioMaterialKnob.Style");
 }
 
 #if WITH_EDITOR
@@ -55,12 +61,92 @@ void UAudioMaterialKnob::SetValue(float InValue)
 	}
 }
 
+void UAudioMaterialKnob::SetTuneSpeed(float InValue)
+{
+	TuneSpeed = FMath::Clamp(InValue, 0.f, 1.0f);
+
+	if (Knob.IsValid())
+	{
+		Knob->SetTuneSpeed(TuneSpeed);
+	}
+}
+
+float UAudioMaterialKnob::GetTuneSpeed() const
+{
+	return TuneSpeed;
+}
+
+void UAudioMaterialKnob::SetFineTuneSpeed(float InValue)
+{
+	FineTuneSpeed = FMath::Clamp(InValue, 0.f, 1.0f);
+
+	if (Knob.IsValid())
+	{
+		Knob->SetFineTuneSpeed(FineTuneSpeed);
+	}
+}
+
+float UAudioMaterialKnob::GetFineTuneSpeed() const
+{
+	return FineTuneSpeed;
+}
+
+void UAudioMaterialKnob::SetLocked(bool InLocked)
+{
+	bLocked = InLocked;
+
+	if (Knob.IsValid())
+	{
+		Knob->SetLocked(InLocked);
+	}
+}
+
+bool UAudioMaterialKnob::GetIsLocked() const
+{
+	return bLocked;
+}
+
+void UAudioMaterialKnob::SetMouseUsesStep(bool InUsesStep)
+{
+	bMouseUsesStep = InUsesStep;
+
+	if (Knob.IsValid())
+	{
+		Knob->SetMouseUsesStep(InUsesStep);
+	}
+}
+
+bool UAudioMaterialKnob::GetMouseUsesStep() const
+{
+	return bMouseUsesStep;
+}
+
+void UAudioMaterialKnob::SetStepSize(float InValue)
+{
+	StepSize = InValue;
+
+	if (Knob.IsValid())
+	{
+		Knob->SetStepSize(InValue);
+	}
+}
+
+float UAudioMaterialKnob::GetStepSize() const
+{
+	return StepSize;
+}
+
 TSharedRef<SWidget> UAudioMaterialKnob::RebuildWidget()
 {
 	Knob = SNew(SAudioMaterialKnob)
 		.Owner(this)
 		.AudioMaterialKnobStyle(&WidgetStyle)
 		.Value(Value)
+		.TuneSpeed(TuneSpeed)
+		.Locked(bLocked)
+		.FineTuneSpeed(FineTuneSpeed)
+		.MouseUsesStep(bMouseUsesStep)
+		.StepSize(StepSize)
 		.OnFloatValueChanged(BIND_UOBJECT_DELEGATE(FOnFloatValueChanged, HandleOnKnobValueChanged));
 		
 	return Knob.ToSharedRef();
diff --git a/Engine/Plugins/Runtime/AudioWidgets/Source/AudioWidgets/Private/AudioMaterialSlate/SAudioMaterialKnob.cpp b/Engine/Plugins/Runtime/AudioWidgets/Source/AudioWidgets/Private/AudioMaterialSlate/SAudioMaterialKnob.cpp
--- a/Engine/Plugins/Runtime/AudioWidgets/Source/AudioWidgets/Private/AudioMaterialSlate/SAudioMaterialKnob.cpp
+++ b/Engine/Plugins/Runtime/AudioWidgets/Source/AudioWidgets/Private/AudioMaterialSlate/SAudioMaterialKnob.cpp
@@ -4,6 +4,7 @@
 #include "AudioMaterialSlate/SAudioMaterialKnob.h"
 #include "AudioMaterialSlate/AudioMaterialSlateTypes.h"
 #include "AudioWidgetsStyle.h"
+#include "Framework/Application/SlateApplication.h"
 #include "SlateOptMacros.h"
 #include "Components/Widget.h"
 
@@ -11,7 +12,16 @@ BEGIN_SLATE_FUNCTION_BUILD_OPTIMIZATION
 void SAudioMaterialKnob::Construct(const FArguments& InArgs)
 {
 	Owner = InArgs._Owner;
+
+	TuneSpeed = InArgs._TuneSpeed;
+	FineTuneSpeed = InArgs._FineTuneSpeed;
+	bIsFocusable = InArgs._IsFocusable;
+	bLocked = InArgs._Locked;
+	bMouseUsesStep = InArgs._MouseUsesStep;
+	StepSize = InArgs._StepSize;
+
 	AudioMaterialKnobStyle = InArgs._AudioMaterialKnobStyle;
+
 	OnValueChanged = InArgs._OnFloatValueChanged;
 	OnMouseCaptureBegin = InArgs._OnMouseCaptureBegin;
 	OnMouseCaptureEnd = InArgs._OnMouseCaptureEnd;
@@ -36,6 +46,36 @@ void SAudioMaterialKnob::SetValue(float InValue)
 	CommitValue(InValue);
 }
 
+void SAudioMaterialKnob::SetTuneSpeed(const float InTurnSpeed)
+{
+	TuneSpeed.Set(InTurnSpeed);
+}
+
+void SAudioMaterialKnob::SetFineTuneSpeed(const float InFineTuneTurnSpeed)
+{
+	FineTuneSpeed.Set(InFineTuneTurnSpeed);
+}
+
+void SAudioMaterialKnob::SetLocked(const bool InLocked)
+{
+	bLocked.Set(InLocked);
+}
+
+void SAudioMaterialKnob::SetMouseUsesStep(const bool InUsesStep)
+{
+	bMouseUsesStep.Set(InUsesStep);
+}
+
+void SAudioMaterialKnob::SetStepSize(const float InStepSize)
+{
+	StepSize.Set(InStepSize);
+}
+
+bool SAudioMaterialKnob::IsLocked() const
+{
+	return bLocked.Get();
+}
+
 UMaterialInstanceDynamic* SAudioMaterialKnob::ApplyNewMaterial()
 {
 	if (AudioMaterialKnobStyle)
@@ -146,15 +186,22 @@ FVector2D SAudioMaterialKnob::ComputeDesiredSize(float) const
 
 FReply SAudioMaterialKnob::OnMouseMove(const FGeometry& MyGeometry, const FPointerEvent& MouseEvent)
 {
-	if (this->HasMouseCapture())
+	if (this->HasMouseCapture() && !IsLocked())
 	{
 		SetCursor(EMouseCursor::GrabHandClosed);
 
 		int32 CurrentYValue = MouseEvent.GetLastScreenSpacePosition().Y;
-		const float MouseSpeed = 0.2f;
+		const float Speed = bIsFineTune ? FineTuneSpeed.Get() : TuneSpeed.Get();
 
-		float ValueDelta = (float)(MouseDownPosition.Y - CurrentYValue) / PixelDelta * MouseSpeed;
+		float ValueDelta = (float)(MouseDownPosition.Y - CurrentYValue) / PixelDelta * Speed;
 		float NewValue = FMath::Clamp(MouseDownValue + ValueDelta, 0.0f, 1.0f);
+
+		if (bMouseUsesStep.Get())
+		{
+			const float SteppedValue = FMath::RoundToInt(NewValue / StepSize.Get()) * StepSize.Get();
+			NewValue = SteppedValue;
+		}
+
 		CommitValue(NewValue);
 
 		return FReply::Handled();
@@ -165,7 +212,7 @@ FReply SAudioMaterialKnob::OnMouseMove(const FGeometry& MyGeometry, const FPoint
 
 FReply SAudioMaterialKnob::OnMouseButtonDown(const FGeometry& MyGeometry, const FPointerEvent& MouseEvent)
 {
-	if ((MouseEvent.GetEffectingButton() == EKeys::LeftMouseButton))
+	if ((MouseEvent.GetEffectingButton() == EKeys::LeftMouseButton) && !IsLocked())
 	{
 		CachedCursor = GetCursor().Get(EMouseCursor::Default);
 
@@ -192,6 +239,37 @@ FReply SAudioMaterialKnob::OnMouseButtonUp(const FGeometry& MyGeometry, const FP
 	return FReply::Unhandled();
 }
 
+FReply SAudioMaterialKnob::OnKeyDown(const FGeometry& MyGeometry, const FKeyEvent& InKeyEvent)
+{
+	if (InKeyEvent.GetKey() == EKeys::LeftShift)
+	{	
+		MouseDownPosition = FSlateApplication::Get().GetCursorPos();
+		MouseDownValue = ValueAttribute.Get();
+		bIsFineTune = true;
+	}
+	
+	return FReply::Unhandled();
+}
+
+FReply SAudioMaterialKnob::OnKeyUp(const FGeometry& MyGeometry, const FKeyEvent& InKeyEvent)
+{
+	bIsFineTune = false;
+	MouseDownPosition = FSlateApplication::Get().GetCursorPos();
+	MouseDownValue = ValueAttribute.Get();
+	
+	return FReply::Unhandled();
+}
+
+bool SAudioMaterialKnob::SupportsKeyboardFocus() const
+{
+	return bIsFocusable.Get();
+}
+
+bool SAudioMaterialKnob::IsInteractable() const
+{
+	return IsEnabled() && !IsLocked() && SupportsKeyboardFocus();;
+}
+
 void SAudioMaterialKnob::CommitValue(float NewValue)
 {
 	const float OldValue = ValueAttribute.Get();
diff --git a/Engine/Plugins/Runtime/AudioWidgets/Source/AudioWidgets/Public/AudioMaterialSlate/AudioMaterialKnob.h b/Engine/Plugins/Runtime/AudioWidgets/Source/AudioWidgets/Public/AudioMaterialSlate/AudioMaterialKnob.h
--- a/Engine/Plugins/Runtime/AudioWidgets/Source/AudioWidgets/Public/AudioMaterialSlate/AudioMaterialKnob.h
+++ b/Engine/Plugins/Runtime/AudioWidgets/Source/AudioWidgets/Public/AudioMaterialSlate/AudioMaterialKnob.h
@@ -25,7 +25,7 @@ class AUDIOWIDGETS_API UAudioMaterialKnob : public UWidget
 
 public:
 
-	UAudioMaterialKnob();
+	UAudioMaterialKnob(const FObjectInitializer& ObjectInitializer);
 
 	/** The button's style */
 	UPROPERTY(EditAnywhere, BlueprintReadWrite, Category = "Style", meta = (DisplayName = "Style", ShowOnlyInnerProperties))
@@ -44,13 +44,53 @@ public:
 	virtual void ReleaseSlateResources(bool bReleaseChildren) override;
 	// End of UVisual
 
-	/** Gets the current value of the knob.*/
+	/** Get the current value of the knob.*/
 	UFUNCTION(BlueprintPure, Category = "Behavior")
 	float GetValue();
 
-	/** Sets the current value of the knob. InValue is Clamped between 0.f - 1.f */
+	/** Set the current value of the knob. InValue is Clamped between 0.f - 1.f */
 	UFUNCTION(BlueprintCallable, Category = "Behavior")
-	void SetValue(float InValue);
+	void SetValue(float InValue);	
+	
+	/** Set the knobs tune speed. InValue is Clamped between 0.f - 1.f */
+	UFUNCTION(BlueprintCallable, Category = "Behavior")
+	void SetTuneSpeed(float InValue);	
+
+	/** Get the Knobs tune speed*/
+	UFUNCTION(BlueprintPure, Category = "Behavior")
+	float GetTuneSpeed() const;
+	
+	/** Set the knobs fine-tune speed. InValue is Clamped between 0.f - 1.f */
+	UFUNCTION(BlueprintCallable, Category = "Behavior")
+	void SetFineTuneSpeed(float InValue);
+
+	/** Get the Knobs fine-tune speed*/
+	UFUNCTION(BlueprintPure, Category = "Behavior")
+	float GetFineTuneSpeed() const;
+	
+	/** Set the knob to be interactive or fixed */
+	UFUNCTION(BlueprintCallable, Category = "Behavior")
+	void SetLocked(bool InLocked);	
+
+	/** Get whether the knob is interactive or fixed.*/
+	UFUNCTION(BlueprintPure, Category = "Behavior")
+	bool GetIsLocked() const;
+	
+	/** Set the knob to use steps when turning On Mouse move */
+	UFUNCTION(BlueprintCallable, Category = "Behavior")
+	void SetMouseUsesStep(bool InUsesStep);
+
+	/** Get whether the knob uses steps when tuning On Mouse move*/
+	UFUNCTION(BlueprintPure, Category = "Behavior")
+	bool GetMouseUsesStep() const;
+
+	/** Set the amount to adjust the value when using steps*/
+	UFUNCTION(BlueprintCallable, Category = "Behavior")
+	void SetStepSize(float InValue);
+
+	/** Get Step Size*/
+	UFUNCTION(BlueprintPure, Category = "Behavior")
+	float GetStepSize() const;
 
 public:
 
@@ -72,6 +112,26 @@ private:
 	UPROPERTY(EditAnywhere, BlueprintSetter = SetValue, BlueprintGetter = GetValue, Category = "Appearance", meta = (UIMin = "0", UIMax = "1"))
 	float Value = 1.f;
 
+	/** The tune speed of the knob */
+	UPROPERTY(EditAnywhere, BlueprintSetter = SetTuneSpeed, BlueprintGetter = GetTuneSpeed, Category = Appearance, meta = (ClampMin = "0", ClampMax = "1", UIMin = "0", UIMax = "1"))
+	float TuneSpeed;
+
+	/** The tune speed when fine-tuning the knob */
+	UPROPERTY(EditAnywhere, BlueprintSetter = SetFineTuneSpeed, BlueprintGetter = GetFineTuneSpeed, Category = Appearance, AdvancedDisplay, meta = (ClampMin = "0", ClampMax = "1", UIMin = "0", UIMax = "1"))
+	float FineTuneSpeed;
+
+	/** Whether the knob is interactive or fixed. */
+	UPROPERTY(EditAnywhere, BlueprintSetter = SetLocked, BlueprintGetter = GetIsLocked, Category = Appearance, AdvancedDisplay)
+	bool bLocked;
+
+	/** Sets new value if mouse position is greater/less than half the step size. */
+	UPROPERTY(EditAnywhere, BlueprintSetter = SetMouseUsesStep, BlueprintGetter = GetMouseUsesStep, Category = Appearance, AdvancedDisplay)
+	bool bMouseUsesStep;
+
+	/** The amount to adjust the value by, when using steps */
+	UPROPERTY(EditAnywhere, BlueprintSetter = SetStepSize, BlueprintGetter = GetStepSize, Category = Appearance, AdvancedDisplay, meta = (UIMin = "0", UIMax = "1"))
+	float StepSize;
+
 private:
 
 	/** Native Slate Widget */
diff --git a/Engine/Plugins/Runtime/AudioWidgets/Source/AudioWidgets/Public/AudioMaterialSlate/SAudioMaterialKnob.h b/Engine/Plugins/Runtime/AudioWidgets/Source/AudioWidgets/Public/AudioMaterialSlate/SAudioMaterialKnob.h
--- a/Engine/Plugins/Runtime/AudioWidgets/Source/AudioWidgets/Public/AudioMaterialSlate/SAudioMaterialKnob.h
+++ b/Engine/Plugins/Runtime/AudioWidgets/Source/AudioWidgets/Public/AudioMaterialSlate/SAudioMaterialKnob.h
@@ -18,6 +18,12 @@ class AUDIOWIDGETS_API SAudioMaterialKnob : public SLeafWidget
 {
 public:
 	SLATE_BEGIN_ARGS(SAudioMaterialKnob)
+	: _TuneSpeed(0.2f)
+	, _FineTuneSpeed(0.05f)
+	, _IsFocusable(true)
+	, _Locked(false)
+	, _MouseUsesStep(false)
+	, _StepSize(0.01f)
 	{}
 
 	/** The owner object*/
@@ -26,6 +32,24 @@ public:
 	/**Value of the Knob*/
 	SLATE_ATTRIBUTE(float, Value)
 
+	/** The tune speed of the knob.*/
+	SLATE_ATTRIBUTE(float, TuneSpeed)
+
+	/** The tune speed of the knob when shift is held. */
+	SLATE_ATTRIBUTE(float, FineTuneSpeed)
+
+	/** When true knob will be keyboard focusable, else only mouse-clickable and never keyboard focusable. */
+	SLATE_ATTRIBUTE(bool, IsFocusable)
+
+	/** Whether the knob is interactive or fixed. */
+	SLATE_ATTRIBUTE(bool, Locked)
+
+	/**Rotates knob in given steps. Sets new value if mouse position is greater/less than half the step size. */
+	SLATE_ATTRIBUTE(bool, MouseUsesStep)
+
+	/** StepSize */
+	SLATE_ATTRIBUTE(float, StepSize)
+
 	/** The style used to draw the knob. */
 	SLATE_STYLE_ARGUMENT(FAudioMaterialKnobStyle, AudioMaterialKnobStyle)
 
@@ -46,6 +70,24 @@ public:
 	/** Set the Value attribute */
 	void SetValue(float InValueAttribute);
 
+	/** Set the TuneSpeed attribute */
+	void SetTuneSpeed(const float InMouseSpeed);
+
+	/** Set the FineTuneSpeed attribute */
+	void SetFineTuneSpeed(const float InMouseFineTuneSpeed);
+
+	/** Set the bLocked attribute */
+	void SetLocked(const bool InLocked);	
+	
+	/** See the bMouseUsesStep attribute */
+	void SetMouseUsesStep(const bool InUsesStep);
+
+	/** Set the StepSize attribute */
+	void SetStepSize(const float InStepSize);
+
+	/** @return Is the knob interaction locked or not?*/
+	bool IsLocked() const;
+
 	/** Apply new material to be used to render the Slate.*/
 	UMaterialInstanceDynamic* ApplyNewMaterial();
 
@@ -77,6 +119,10 @@ protected:
 	virtual FReply OnMouseMove(const FGeometry& MyGeometry, const FPointerEvent& MouseEvent) override;
 	virtual FReply OnMouseButtonDown(const FGeometry& MyGeometry, const FPointerEvent& MouseEvent) override;
 	virtual FReply OnMouseButtonUp(const FGeometry& MyGeometry, const FPointerEvent& MouseEvent)override;
+	virtual FReply OnKeyDown(const FGeometry& MyGeometry, const FKeyEvent& InKeyEvent) override;
+	virtual FReply OnKeyUp(const FGeometry& MyGeometry, const FKeyEvent& InKeyEvent) override;
+	virtual bool SupportsKeyboardFocus() const override;
+	virtual bool IsInteractable() const override;
 	//~SWidget
 
 protected:
@@ -107,6 +153,24 @@ private:
 	//Holds the knobs current Value
 	TAttribute<float> ValueAttribute = 1.0;
 
+	/** Holds the amount to adjust the knob On Mouse move*/
+	TAttribute<float> TuneSpeed;
+
+	/** Holds the amount to adjust the knob On Mouse move & FineTuning */
+	TAttribute<float> FineTuneSpeed;
+
+	/** Holds a flag indicating whether knob will be keyboard focusable. */
+	TAttribute<bool> bIsFocusable;
+
+	// Holds a flag indicating whether the knob is locked.
+	TAttribute<bool> bLocked;	
+	
+	// Holds a flag indicating whether the knob uses steps when roating on Mouse move.
+	TAttribute<bool> bMouseUsesStep;
+
+	/** Holds the amount to adjust the value when steps are used */
+	TAttribute<float> StepSize;
+
 	// The position of the mouse when it pushed down and started rotating the knob
 	FVector2D MouseDownPosition;
 
@@ -119,4 +183,7 @@ private:
 	// the max pixels to go to min or max value (clamped to 0 or 1) in one drag period
 	int32 PixelDelta = 50;
 
+	// Whether or not we're in fine-tune mode
+	bool bIsFineTune;
+
 };
