Extra Features to the AudioMaterialSlider. Fine-tuning, Locking, StepSize
#rnx

diff --git a/Engine/Plugins/Runtime/AudioWidgets/Source/AudioWidgets/Private/AudioMaterialSlate/AudioMaterialSlider.cpp b/Engine/Plugins/Runtime/AudioWidgets/Source/AudioWidgets/Private/AudioMaterialSlate/AudioMaterialSlider.cpp
--- a/Engine/Plugins/Runtime/AudioWidgets/Source/AudioWidgets/Private/AudioMaterialSlate/AudioMaterialSlider.cpp
+++ b/Engine/Plugins/Runtime/AudioWidgets/Source/AudioWidgets/Private/AudioMaterialSlate/AudioMaterialSlider.cpp
@@ -9,7 +9,13 @@
 
 #define LOCTEXT_NAMESPACE "AudioWidgets"
 
-UAudioMaterialSlider::UAudioMaterialSlider()
+UAudioMaterialSlider::UAudioMaterialSlider(const FObjectInitializer& ObjectInitializer)
+	: Super(ObjectInitializer)
+	, TuneSpeed(0.2f)
+	, FineTuneSpeed(0.05f)
+	, bLocked(false)
+	, bMouseUsesStep(false)
+	, StepSize(0.01f)
 {
 	//get default style
 	WidgetStyle = FAudioWidgetsStyle::Get().GetWidgetStyle<FAudioMaterialSliderStyle>("AudioMaterialSlider.Style");
@@ -62,11 +68,91 @@ void UAudioMaterialSlider::SetValue(float InValue)
 	}
 }
 
+void UAudioMaterialSlider::SetTuneSpeed(const float InValue)
+{
+	TuneSpeed = FMath::Clamp(InValue, 0.f, 1.f);
+
+	if (Slider.IsValid())
+	{
+		Slider.Get()->SetTuneSpeed(TuneSpeed);
+	}
+}
+
+float UAudioMaterialSlider::GetTuneSpeed() const
+{
+	return TuneSpeed;
+}
+
+void UAudioMaterialSlider::SetFineTuneSpeed(const float InValue)
+{
+	FineTuneSpeed = FMath::Clamp(InValue, 0.f, 1.0f);
+
+	if (Slider.IsValid())
+	{
+		Slider.Get()->SetFineTuneSpeed(FineTuneSpeed);
+	}
+}
+
+float UAudioMaterialSlider::GetFineTuneSpeed() const
+{
+	return FineTuneSpeed;
+}
+
+void UAudioMaterialSlider::SetLocked(bool bInLocked)
+{
+	bLocked = bInLocked;
+
+	if (Slider.IsValid())
+	{
+		Slider->SetLocked(bInLocked);
+	}
+}
+
+bool UAudioMaterialSlider::GetIsLocked() const
+{
+	return bLocked;
+}
+
+void UAudioMaterialSlider::SetMouseUsesStep(bool bInUsesStep)
+{
+	bMouseUsesStep = bInUsesStep;
+
+	if (Slider.IsValid())
+	{
+		Slider->SetMouseUsesStep(bInUsesStep);
+	}
+}
+
+bool UAudioMaterialSlider::GetMouseUsesStep() const
+{
+	return bMouseUsesStep;
+}
+
+void UAudioMaterialSlider::SetStepSize(float InValue)
+{
+	StepSize = InValue;
+
+	if (Slider.IsValid())
+	{
+		Slider->SetStepSize(InValue);
+	}
+}
+
+float UAudioMaterialSlider::GetStepSize() const
+{
+	return StepSize;
+}
+
 TSharedRef<SWidget> UAudioMaterialSlider::RebuildWidget()
 {
 	Slider = SNew(SAudioMaterialSlider)
 		.Owner(this)
 		.Orientation(Orientation)
+		.TuneSpeed(TuneSpeed)
+		.FineTuneSpeed(FineTuneSpeed)
+		.Locked(bLocked)
+		.MouseUsesStep(bMouseUsesStep)
+		.StepSize(StepSize)
 		.AudioMaterialSliderStyle(&WidgetStyle)
 		.OnValueChanged(BIND_UOBJECT_DELEGATE(FOnFloatValueChanged, HandleOnValueChanged));
 
diff --git a/Engine/Plugins/Runtime/AudioWidgets/Source/AudioWidgets/Private/AudioMaterialSlate/SAudioMaterialSlider.cpp b/Engine/Plugins/Runtime/AudioWidgets/Source/AudioWidgets/Private/AudioMaterialSlate/SAudioMaterialSlider.cpp
--- a/Engine/Plugins/Runtime/AudioWidgets/Source/AudioWidgets/Private/AudioMaterialSlate/SAudioMaterialSlider.cpp
+++ b/Engine/Plugins/Runtime/AudioWidgets/Source/AudioWidgets/Private/AudioMaterialSlate/SAudioMaterialSlider.cpp
@@ -4,6 +4,7 @@
 #include "AudioMaterialSlate/SAudioMaterialSlider.h"
 #include "AudioMaterialSlate/AudioMaterialSlider.h"
 #include "AudioWidgetsStyle.h"
+#include "Framework/Application/SlateApplication.h"
 #include "SlateOptMacros.h"
 #include "Styling/SlateBrush.h"
 
@@ -12,8 +13,17 @@ BEGIN_SLATE_FUNCTION_BUILD_OPTIMIZATION
 void SAudioMaterialSlider::Construct(const FArguments& InArgs)
 {
 	Owner = InArgs._Owner;
+
+	TuneSpeed = InArgs._TuneSpeed;
+	FineTuneSpeed = InArgs._FineTuneSpeed;
+	bIsFocusable = InArgs._IsFocusable;
 	Orientation = InArgs._Orientation;
+	bMouseUsesStep = InArgs._MouseUsesStep;
+	bLocked = InArgs._Locked;
+	StepSize = InArgs._StepSize;
+
 	AudioMaterialSliderStyle = InArgs._AudioMaterialSliderStyle;
+
 	OnValueChanged = InArgs._OnValueChanged;
 	OnValueCommitted = InArgs._OnValueCommitted;
 
@@ -113,6 +123,36 @@ void SAudioMaterialSlider::SetValue(TAttribute<float> InValueAttribute)
 	CommitValue(InValueAttribute.Get());
 }
 
+void SAudioMaterialSlider::SetTuneSpeed(const float InMouseTuneSpeed)
+{
+	TuneSpeed.Set(InMouseTuneSpeed);
+}
+
+void SAudioMaterialSlider::SetFineTuneSpeed(const float InMouseFineTuneSpeed)
+{
+	FineTuneSpeed.Set(InMouseFineTuneSpeed);
+}
+
+void SAudioMaterialSlider::SetMouseUsesStep(const bool InUsesStep)
+{
+	bMouseUsesStep.Set(InUsesStep);
+}
+
+void SAudioMaterialSlider::SetStepSize(const float InStepSize)
+{
+	StepSize.Set(InStepSize);
+}
+
+void SAudioMaterialSlider::SetLocked(const bool bInLocked)
+{
+	bLocked.Set(bInLocked);
+}
+
+bool SAudioMaterialSlider::IsLocked() const
+{
+	return bLocked.Get();
+}
+
 void SAudioMaterialSlider::ApplyNewMaterial()
 {
 	if (AudioMaterialSliderStyle)
@@ -132,12 +172,12 @@ void SAudioMaterialSlider::SetOrientation(EOrientation InOrientation)
 
 FReply SAudioMaterialSlider::OnMouseButtonDown(const FGeometry& MyGeometry, const FPointerEvent& MouseEvent)
 {
-	if ((MouseEvent.GetEffectingButton() == EKeys::LeftMouseButton))
-	{
-		FVector2D Value2D = PositionToValue(MyGeometry, MouseEvent.GetLastScreenSpacePosition());
-		const float Value = Orientation == Orient_Horizontal ? Value2D.X : Value2D.Y;
-		CommitValue(Value);
+	if ((MouseEvent.GetEffectingButton() == EKeys::LeftMouseButton) && !IsLocked())
+	{	
+		CachedCursor = GetCursor().Get(EMouseCursor::Default);
 
+		MouseDownStartPosition = MouseEvent.GetScreenSpacePosition();
+		MouseDownValue = ValueAttribute.Get();
 		return FReply::Handled().CaptureMouse(SharedThis(this));
 	}
 
@@ -148,6 +188,7 @@ FReply SAudioMaterialSlider::OnMouseButtonUp(const FGeometry& MyGeometry, const
 {
 	if ((MouseEvent.GetEffectingButton() == EKeys::LeftMouseButton) && HasMouseCapture())
 	{
+		SetCursor(CachedCursor);
 		return FReply::Handled().ReleaseMouseCapture();
 	}
 
@@ -156,11 +197,25 @@ FReply SAudioMaterialSlider::OnMouseButtonUp(const FGeometry& MyGeometry, const
 
 FReply SAudioMaterialSlider::OnMouseMove(const FGeometry& MyGeometry, const FPointerEvent& MouseEvent)
 {
-	if (this->HasMouseCapture())
+	if (this->HasMouseCapture() && !IsLocked())
 	{
-		FVector2D Value2D = PositionToValue(MyGeometry, MouseEvent.GetLastScreenSpacePosition());
-		const float Value = Orientation == Orient_Horizontal ? Value2D.X : Value2D.Y;
-		CommitValue(Value);
+		SetCursor(EMouseCursor::GrabHandClosed);
+
+		int32 MouseCurrentPosition = Orientation == Orient_Horizontal ? MouseEvent.GetLastScreenSpacePosition().X : MouseEvent.GetLastScreenSpacePosition().Y;
+		float MouseDownPosition = Orientation == Orient_Horizontal ? MouseDownStartPosition.X : MouseDownStartPosition.Y;
+		const float Speed = bIsFineTune ? FineTuneSpeed.Get() : TuneSpeed.Get();
+
+		float ValueDelta = (float)(MouseDownPosition - MouseCurrentPosition) / PixelDelta * Speed;
+
+		float NewValue = Orientation == Orient_Horizontal ? FMath::Clamp(MouseDownValue - ValueDelta, 0.0f, 1.0f) : FMath::Clamp(MouseDownValue + ValueDelta, 0.0f, 1.0f);
+
+		if (bMouseUsesStep.Get())
+		{
+			const float SteppedValue = FMath::RoundToInt(NewValue / StepSize.Get()) * StepSize.Get();
+			NewValue = SteppedValue;
+		}
+
+		CommitValue(NewValue);
 
 		return FReply::Handled();
 	}
@@ -173,6 +228,38 @@ void SAudioMaterialSlider::OnMouseCaptureLost(const FCaptureLostEvent& CaptureLo
 	OnValueCommitted.ExecuteIfBound(ValueAttribute.Get());
 	SLeafWidget::OnMouseCaptureLost(CaptureLostEvent);
 }
+
+FReply SAudioMaterialSlider::OnKeyDown(const FGeometry& MyGeometry, const FKeyEvent& InKeyEvent)
+{
+	if (InKeyEvent.GetKey() == EKeys::LeftShift)
+	{
+		MouseDownStartPosition = FSlateApplication::Get().GetCursorPos();
+		MouseDownValue = ValueAttribute.Get();
+		bIsFineTune = true;
+	}
+
+	return FReply::Unhandled();
+}
+
+FReply SAudioMaterialSlider::OnKeyUp(const FGeometry& MyGeometry, const FKeyEvent& InKeyEvent)
+{
+	bIsFineTune = false;
+	MouseDownStartPosition = FSlateApplication::Get().GetCursorPos();
+	MouseDownValue = ValueAttribute.Get();
+
+	return FReply::Unhandled();
+}
+
+bool SAudioMaterialSlider::SupportsKeyboardFocus() const
+{
+	return bIsFocusable.Get();
+}
+
+bool SAudioMaterialSlider::IsInteractable() const
+{
+	return IsEnabled() && !IsLocked() && SupportsKeyboardFocus();
+}
+
 void SAudioMaterialSlider::CommitValue(float NewValue)
 {
 	const float OldValue = ValueAttribute.Get();
@@ -189,16 +276,3 @@ void SAudioMaterialSlider::CommitValue(float NewValue)
 		OnValueChanged.ExecuteIfBound(Val);
 	}
 }
-
-FVector2D SAudioMaterialSlider::PositionToValue(const FGeometry& MyGeometry, const FVector2D& AbsolutePosition)
-{
-	const FVector2D LocalPosition = MyGeometry.AbsoluteToLocal(AbsolutePosition);
-
-	float RelativeValueX = (LocalPosition.X) / (MyGeometry.Size.X);
-	float RelativeValueY = 1 - (LocalPosition.Y) / (MyGeometry.Size.Y);
-
-	RelativeValueX = FMath::Clamp(RelativeValueX, 0.0f, 1.0f);
-	RelativeValueY = FMath::Clamp(RelativeValueY, 0.0f, 1.0f);
-
-	return FVector2D(RelativeValueX, RelativeValueY);
-}
diff --git a/Engine/Plugins/Runtime/AudioWidgets/Source/AudioWidgets/Public/AudioMaterialSlate/AudioMaterialSlider.h b/Engine/Plugins/Runtime/AudioWidgets/Source/AudioWidgets/Public/AudioMaterialSlate/AudioMaterialSlider.h
--- a/Engine/Plugins/Runtime/AudioWidgets/Source/AudioWidgets/Public/AudioMaterialSlate/AudioMaterialSlider.h
+++ b/Engine/Plugins/Runtime/AudioWidgets/Source/AudioWidgets/Public/AudioMaterialSlate/AudioMaterialSlider.h
@@ -24,7 +24,7 @@ class AUDIOWIDGETS_API UAudioMaterialSlider : public UWidget
 
 public:
 
-	UAudioMaterialSlider();
+	UAudioMaterialSlider(const FObjectInitializer& ObjectInitializer);
 
 	/** The slider's style */
 	UPROPERTY(EditAnywhere, BlueprintReadWrite, Category = "Style", meta = (DisplayName = "Style", ShowOnlyInnerProperties))
@@ -52,6 +52,46 @@ public:
 	UFUNCTION(BlueprintCallable, Category = "Behavior")
 	void SetValue(float InValue);
 
+	/** Set the tune speed of the slider. InValue is Clamped between 0.f - 1.f */
+	UFUNCTION(BlueprintCallable, Category = "Behavior")
+	void SetTuneSpeed(const float InValue);
+
+	/** Get slider tune speed*/
+	UFUNCTION(BlueprintPure, Category = "Behavior")
+	float GetTuneSpeed() const;
+
+	/** Set the fine-tune speed of the slider. InValue is Clamped between 0.f - 1.f */
+	UFUNCTION(BlueprintCallable, Category = "Behavior")
+	void SetFineTuneSpeed(const float InValue);
+
+	/** Get slider fine-tune speed*/
+	UFUNCTION(BlueprintPure, Category = "Behavior")
+	float GetFineTuneSpeed() const;
+
+	/** Set the slider to be interactive or fixed */
+	UFUNCTION(BlueprintCallable, Category = "Behavior")
+	void SetLocked(bool bInLocked);
+
+	/** Get whether the slider is interactive or fixed.*/
+	UFUNCTION(BlueprintPure, Category = "Behavior")
+	bool GetIsLocked() const;
+
+	/** Sets the slider to use steps when turning On Mouse move */
+	UFUNCTION(BlueprintCallable, Category = "Behavior")
+	void SetMouseUsesStep(bool bInUsesStep);
+
+	/** Get whether the slider uses steps when turning On Mouse move*/
+	UFUNCTION(BlueprintPure, Category = "Behavior")
+	bool GetMouseUsesStep() const;
+
+	/** Sets the amount to adjust the value when using steps*/
+	UFUNCTION(BlueprintCallable, Category = "Behavior")
+	void SetStepSize(float InValue);
+
+	/** Get Step Size*/
+	UFUNCTION(BlueprintPure, Category = "Behavior")
+	float GetStepSize() const;
+
 public:
 
 	/** Called when the value is changed by slider. */
@@ -76,6 +116,26 @@ private:
 	UPROPERTY(EditAnywhere, Category = "Appearance")
 	TEnumAsByte<EOrientation> Orientation = EOrientation::Orient_Horizontal;
 
+	/** The tune speed of the slider On Mouse move */
+	UPROPERTY(EditAnywhere, BlueprintSetter = SetTuneSpeed, BlueprintGetter = GetTuneSpeed, Category = Appearance, meta = (ClampMin = "0", ClampMax = "1", UIMin = "0", UIMax = "1"))
+	float TuneSpeed;
+
+	/** The tune speed of the slider when fine-tuning the slider On Mouse move && Left-Shift pressed */
+	UPROPERTY(EditAnywhere, BlueprintSetter = SetFineTuneSpeed, BlueprintGetter = GetFineTuneSpeed, Category = Appearance, AdvancedDisplay, meta = (ClampMin = "0", ClampMax = "1", UIMin = "0", UIMax = "1"))
+	float FineTuneSpeed;
+
+	/** Whether the slider is interactive or fixed. */
+	UPROPERTY(EditAnywhere, BlueprintSetter = SetLocked, BlueprintGetter = GetIsLocked, Category = Appearance, AdvancedDisplay)
+	bool bLocked;
+
+	/** Sets new value if mouse position is greater/less than half the step size. */
+	UPROPERTY(EditAnywhere, BlueprintSetter = SetMouseUsesStep, BlueprintGetter = GetMouseUsesStep, Category = Appearance, AdvancedDisplay)
+	bool bMouseUsesStep;
+
+	/** The amount to adjust the value by, when using steps */
+	UPROPERTY(EditAnywhere, BlueprintSetter = SetStepSize, BlueprintGetter = GetStepSize, Category = Appearance, AdvancedDisplay, meta = (UIMin = "0", UIMax = "1"))
+	float StepSize;
+
 private:
 
 	/** Native Slate Widget */
diff --git a/Engine/Plugins/Runtime/AudioWidgets/Source/AudioWidgets/Public/AudioMaterialSlate/SAudioMaterialSlider.h b/Engine/Plugins/Runtime/AudioWidgets/Source/AudioWidgets/Public/AudioMaterialSlate/SAudioMaterialSlider.h
--- a/Engine/Plugins/Runtime/AudioWidgets/Source/AudioWidgets/Public/AudioMaterialSlate/SAudioMaterialSlider.h
+++ b/Engine/Plugins/Runtime/AudioWidgets/Source/AudioWidgets/Public/AudioMaterialSlate/SAudioMaterialSlider.h
@@ -18,6 +18,12 @@ class AUDIOWIDGETS_API SAudioMaterialSlider : public SLeafWidget
 {
 public:
 	SLATE_BEGIN_ARGS(SAudioMaterialSlider)
+	: _TuneSpeed(0.2f)
+	, _FineTuneSpeed(0.05f)
+	, _IsFocusable(true)
+	, _Locked(false)
+	, _MouseUsesStep(false)
+	, _StepSize(0.01f)
 	{}
 
 	/** The owner object*/
@@ -26,6 +32,24 @@ public:
 	/** The slider's orientation. */
 	SLATE_ARGUMENT(EOrientation, Orientation)
 
+	/** The tune speed of the slider handle */
+	SLATE_ATTRIBUTE(float, TuneSpeed)
+
+	/** The fine-tune speed of the slider handle */
+	SLATE_ATTRIBUTE(float, FineTuneSpeed)
+
+	/** When true slider will be keyboard focusable, else only mouse-clickable and never keyboard focusable. */
+	SLATE_ATTRIBUTE(bool, IsFocusable)
+
+	/** Whether the slider is interactive or fixed. */
+	SLATE_ATTRIBUTE(bool, Locked)
+
+	/**Moves Slider handle in given steps. Sets new value if mouse position is greater/less than half the step size. */
+	SLATE_ATTRIBUTE(bool, MouseUsesStep)
+
+	/** StepSize */
+	SLATE_ATTRIBUTE(float, StepSize)
+
 	/** The style used to draw the slider. */
 	SLATE_STYLE_ARGUMENT(FAudioMaterialSliderStyle, AudioMaterialSliderStyle)
 
@@ -52,6 +76,24 @@ public:
 	/** Set the Value attribute */
 	void SetValue(TAttribute<float> InValueAttribute);
 
+	/** Set the TuneSpeed attribute */
+	void SetTuneSpeed(const float InMouseTuneSpeed);
+
+	/** Set the FineTuneSpeed attribute */
+	void SetFineTuneSpeed(const float InMouseFineTuneSpeed);
+
+	/** See the bMouseUsesStep attribute */
+	void SetMouseUsesStep(const bool InUsesStep);
+
+	/** Set the StepSize attribute */
+	void SetStepSize(const float InStepSize);
+
+	/** Set the bLocked attribute */
+	void SetLocked(const bool bInLocked);
+
+	/** @return Is the knob interaction locked or not?*/
+	bool IsLocked() const;
+
 	/** Apply new material to be used to render the Slate.*/
 	void ApplyNewMaterial();
 
@@ -63,6 +105,10 @@ public:
 	virtual FReply OnMouseButtonUp(const FGeometry& MyGeometry, const FPointerEvent& MouseEvent) override;
 	virtual FReply OnMouseMove(const FGeometry& MyGeometry, const FPointerEvent& MouseEvent) override;
 	virtual void OnMouseCaptureLost(const FCaptureLostEvent& CaptureLostEvent) override;
+	virtual FReply OnKeyDown(const FGeometry& MyGeometry, const FKeyEvent& InKeyEvent) override;
+	virtual FReply OnKeyUp(const FGeometry& MyGeometry, const FKeyEvent& InKeyEvent) override;
+	virtual bool SupportsKeyboardFocus() const override;
+	virtual bool IsInteractable() const override;
 	//~SWidget
 
 public:
@@ -82,15 +128,6 @@ private:
 	*/
 	void CommitValue(float NewValue);
 
-	/**
-	* Calculates the new value based on the given absolute coordinates.
-	*
-	* @param MyGeometry The slider's geometry.
-	* @param AbsolutePosition The absolute position of the slider.
-	* @return The new value.
-	*/
-	FVector2D PositionToValue(const FGeometry& MyGeometry, const FVector2D& AbsolutePosition);
-
 private:
 
 	// Optional override for desired size 
@@ -111,4 +148,36 @@ private:
 	//Holds the current value
 	TAttribute<float> ValueAttribute = 0.f;
 
+	/** Holds the amount to adjust the slider On Mouse move*/
+	TAttribute<float> TuneSpeed;
+
+	/** Holds the amount to adjust the slider On Mouse move & FineTuning */
+	TAttribute<float> FineTuneSpeed;
+
+	/** Holds a flag indicating whether slider will be keyboard focusable. */
+	TAttribute<bool> bIsFocusable;
+
+	// Holds a flag indicating whether the slider is locked.
+	TAttribute<bool> bLocked;
+
+	// Holds a flag indicating whether the slider uses steps when roating on Mouse move.
+	TAttribute<bool> bMouseUsesStep;
+
+	/** Holds the amount to adjust the value when steps are used */
+	TAttribute<float> StepSize;
+
+	// the max pixels to go to min or max value (clamped to 0 or 1) in one drag period
+	int32 PixelDelta = 50;
+
+	// Whether or not we're in fine-tune mode	
+	bool bIsFineTune = false;
+
+	// The position of the mouse when it pushed down and started moving the slider
+	FVector2D MouseDownStartPosition;
+
+	// the value when the mouse was pushed down
+	float MouseDownValue = 0.f;
+
+	// Holds the initial cursor in case a custom cursor has been specified, so we can restore it after dragging the slider
+	EMouseCursor::Type CachedCursor = EMouseCursor::None;
 };
