AudioMaterialSlider Label Unit text change and value processing depending on the selected Unit.
#rnx

diff --git a/Engine/Plugins/Runtime/AudioWidgets/Source/AudioWidgets/Private/AudioMaterialSlate/SAudioMaterialLabeledSlider.cpp b/Engine/Plugins/Runtime/AudioWidgets/Source/AudioWidgets/Private/AudioMaterialSlate/SAudioMaterialLabeledSlider.cpp
--- a/Engine/Plugins/Runtime/AudioWidgets/Source/AudioWidgets/Private/AudioMaterialSlate/SAudioMaterialLabeledSlider.cpp
+++ b/Engine/Plugins/Runtime/AudioWidgets/Source/AudioWidgets/Private/AudioMaterialSlate/SAudioMaterialLabeledSlider.cpp
@@ -12,6 +12,7 @@ void SAudioMaterialLabeledSlider::Construct(const SAudioMaterialLabeledSlider::F
 	OnValueChanged = InArgs._OnValueChanged;
 	OnValueCommitted = InArgs._OnValueCommitted;
 	Orientation = InArgs._Orientation;
+	AudioUnitsValueType = InArgs._AudioUnitsValueType;
 	DesiredSizeOverride = InArgs._DesiredSizeOverride;
 
 	if (InArgs._SliderValue.IsSet())
@@ -48,6 +49,33 @@ void SAudioMaterialLabeledSlider::Construct(const SAudioMaterialLabeledSlider::F
 			Label->SetValueText(OutputValue);
 		});
 
+	//Set the Unit processor from Units Type
+	switch (AudioUnitsValueType.Get())
+	{
+		case EAudioUnitsValueType::Linear:
+			if (Label.IsValid())
+			{
+				AudioUnitProcessor = MakeShared<FAudioUnitProcessor>();
+				Label.Get()->SetShowUnitsText(false);
+			}
+			break;
+		case EAudioUnitsValueType::Frequency:
+			AudioUnitProcessor = MakeShared<FFrequencyProcessor>();
+			break;
+		case EAudioUnitsValueType::Volume:
+			AudioUnitProcessor = MakeShared<FVolumeProcessor>(InArgs._bUseLinearOutput);
+			break;
+	}
+
+	if (AudioUnitProcessor.IsValid())
+	{
+		SetOutputRange(AudioUnitProcessor->GetDefaultOutputRange());
+		if (Label.IsValid())
+		{
+			Label.Get()->SetUnitsText(AudioUnitProcessor->GetUnitsText());
+		}
+	}
+
 	ChildSlot
 	[
 		CreateWidgetLayout()
@@ -92,28 +120,31 @@ void SAudioMaterialLabeledSlider::SetDesiredSizeOverride(const FVector2D Size)
 }
 
 const float SAudioMaterialLabeledSlider::GetOutputValue(const float InSliderValue)
-{
-	return FMath::GetMappedRangeValueClamped(FVector2D(0.0f, 1.0f), OutputRange, InSliderValue);
+{	
+	return AudioUnitProcessor.IsValid() == true ? AudioUnitProcessor->GetOutputValue(OutputRange, InSliderValue) : 0.0f;
 }
 
 const float SAudioMaterialLabeledSlider::GetOutputValueForText(const float InSliderValue)
 {
-	return GetOutputValue(InSliderValue);
+	return AudioUnitProcessor.IsValid() == true ? AudioUnitProcessor->GetOutputValueForText(OutputRange, InSliderValue) : 0.0f;
 }
 
 const float SAudioMaterialLabeledSlider::GetSliderValueForText(const float OutputValue)
 {
-	return GetSliderValue(OutputValue);
+	return AudioUnitProcessor.IsValid() == true ? AudioUnitProcessor->GetSliderValueForText(OutputRange, OutputValue) : 0.0f;
 }
 
 const float SAudioMaterialLabeledSlider::GetSliderValue(const float OutputValue)
 {
-	return FMath::GetMappedRangeValueClamped(OutputRange, FVector2D(0.0f, 1.0f) , OutputValue);
+	return AudioUnitProcessor.IsValid() == true ? AudioUnitProcessor->GetSliderValue(OutputRange, OutputValue) : 0.0f;
 }
 
-void SAudioMaterialLabeledSlider::SetOutputRange(const FVector2D Range)
-{
+void SAudioMaterialLabeledSlider::SetOutputRange(const FVector2D InRange)
+{	
+	//Check the valid range from the processor
+	FVector2D Range = AudioUnitProcessor.IsValid() == true ? AudioUnitProcessor->GetOutputRange(InRange) : InRange;
 	OutputRange = Range;
+
 	// if Range.Y < Range.X, set Range.X to Range.Y
 	OutputRange.X = FMath::Min(Range.X, Range.Y);
 
diff --git a/Engine/Plugins/Runtime/AudioWidgets/Source/AudioWidgets/Private/AudioMaterialSlate/SAudioMaterialSlider.cpp b/Engine/Plugins/Runtime/AudioWidgets/Source/AudioWidgets/Private/AudioMaterialSlate/SAudioMaterialSlider.cpp
--- a/Engine/Plugins/Runtime/AudioWidgets/Source/AudioWidgets/Private/AudioMaterialSlate/SAudioMaterialSlider.cpp
+++ b/Engine/Plugins/Runtime/AudioWidgets/Source/AudioWidgets/Private/AudioMaterialSlate/SAudioMaterialSlider.cpp
@@ -120,7 +120,7 @@ FVector2D SAudioMaterialSlider::ComputeDesiredSize(float) const
 
 void SAudioMaterialSlider::SetValue(TAttribute<float> InValueAttribute)
 {
-	CommitValue(InValueAttribute.Get());
+	ValueAttribute = InValueAttribute;
 }
 
 void SAudioMaterialSlider::SetTuneSpeed(const float InMouseTuneSpeed)
diff --git a/Engine/Plugins/Runtime/AudioWidgets/Source/AudioWidgets/Private/SAudioInputWidget.cpp b/Engine/Plugins/Runtime/AudioWidgets/Source/AudioWidgets/Private/SAudioInputWidget.cpp
new file mode 100644
--- /dev/null
+++ b/Engine/Plugins/Runtime/AudioWidgets/Source/AudioWidgets/Private/SAudioInputWidget.cpp
@@ -0,0 +1,147 @@
+// Copyright Epic Games, Inc. All Rights Reserved.
+
+
+#include "SAudioInputWidget.h"
+#include "AudioDefines.h"
+#include "DSP/Dsp.h"
+
+const FVector2D FAudioUnitProcessor::NormalizedLinearSliderRange = FVector2D(0.0f, 1.0f);
+
+const float FAudioUnitProcessor::GetOutputValue(const FVector2D OutputRange, const float InSliderValue)
+{
+	return FMath::GetMappedRangeValueClamped(FVector2D(0.0f, 1.0f), OutputRange, InSliderValue);
+};
+
+const float FAudioUnitProcessor::GetOutputValueForText(const FVector2D OutputRange, const float InSliderValue)
+{
+	return GetOutputValue(OutputRange, InSliderValue);
+};
+
+const float FAudioUnitProcessor::GetSliderValue(const FVector2D OutputRange, const float OutputValue)
+{
+	return FMath::GetMappedRangeValueClamped(OutputRange, FVector2D(0.0f, 1.0f), OutputValue);
+};
+
+const float FAudioUnitProcessor::GetSliderValueForText(const FVector2D OutputRange, const float OutputValue)
+{
+	return GetSliderValue(OutputRange, OutputValue);
+};
+
+const FText FVolumeProcessor::GetUnitsText()
+{
+	return FText::FromString("dB");
+}
+
+const FVector2D FVolumeProcessor::GetDefaultOutputRange()
+{
+	if (bUseLinearOutput)
+	{
+		return NormalizedLinearSliderRange;
+	}
+	
+	return FVector2D(-100.0f, 0.0f);
+}
+
+const float FVolumeProcessor::GetOutputValue(const FVector2D OutputRange, const float InSliderValue)
+{
+	if (bUseLinearOutput)
+	{
+		// Return linear given normalized linear 
+		const FVector2D LinearSliderRange = FVector2D(Audio::ConvertToLinear(OutputRange.X), Audio::ConvertToLinear(OutputRange.Y));
+		return FMath::GetMappedRangeValueClamped(NormalizedLinearSliderRange, LinearSliderRange, InSliderValue);
+	}
+	else
+	{
+		return GetDbValueFromSliderValue(OutputRange, InSliderValue);
+	}
+}
+
+const float FVolumeProcessor::GetOutputValueForText(const FVector2D OutputRange, const float InSliderValue)
+{
+	return GetDbValueFromSliderValue(OutputRange, InSliderValue);
+}
+
+const float FVolumeProcessor::GetSliderValue(const FVector2D OutputRange, const float OutputValue)
+{
+	if (bUseLinearOutput)
+	{
+		// Convert from linear to normalized linear 
+		const FVector2D LinearSliderRange = FVector2D(Audio::ConvertToLinear(OutputRange.X), Audio::ConvertToLinear(OutputRange.Y));
+		return FMath::GetMappedRangeValueClamped(LinearSliderRange, NormalizedLinearSliderRange, OutputValue);
+	}
+	else
+	{
+		return GetSliderValueFromDb(OutputRange, OutputValue);
+	}
+}
+
+const float FVolumeProcessor::GetSliderValueForText(const FVector2D OutputRange, const float OutputValue)
+{
+	return GetSliderValueFromDb(OutputRange, OutputValue);;
+}
+
+const float FVolumeProcessor::MinDbValue = -160.0f;
+const float FVolumeProcessor::MaxDbValue = 770.0f;
+const FVector2D FVolumeProcessor::GetOutputRange(const FVector2D InRange)
+{
+	// For volume slider, OutputRange is always in dB
+	if (bUseLinearOutput)
+	{
+		// If using linear output, assume given range is linear (not normalized though) 
+		FVector2D RangeInDecibels = FVector2D(Audio::ConvertToDecibels(InRange.X), Audio::ConvertToDecibels(InRange.Y));
+		return FVector2D(FMath::Max(MinDbValue, RangeInDecibels.X), FMath::Min(MaxDbValue, RangeInDecibels.Y));
+	}
+
+	return FVector2D(FMath::Max(MinDbValue, InRange.X), FMath::Min(MaxDbValue, InRange.Y));			
+}
+
+const float FVolumeProcessor::GetDbValueFromSliderValue(const FVector2D OutputRange, const float InSliderValue)
+{
+	// convert from linear 0-1 space to decibel OutputRange that has been converted to linear 
+	const FVector2D LinearSliderRange = FVector2D(Audio::ConvertToLinear(OutputRange.X), Audio::ConvertToLinear(OutputRange.Y));
+	const float LinearSliderValue = FMath::GetMappedRangeValueClamped(NormalizedLinearSliderRange, LinearSliderRange, InSliderValue);
+	// convert from linear to decibels 
+	float OutputValue = Audio::ConvertToDecibels(LinearSliderValue);
+
+	return FMath::Clamp(OutputValue, OutputRange.X, OutputRange.Y);
+}
+
+const float FVolumeProcessor::GetSliderValueFromDb(const FVector2D OutputRange, const float DbValue)
+{
+	float ClampedValue = FMath::Clamp(DbValue, OutputRange.X, OutputRange.Y);
+	// convert from decibels to linear
+	float LinearSliderValue = Audio::ConvertToLinear(ClampedValue);
+	// convert from decibel OutputRange that has been converted to linear to linear 0-1 space 
+	const FVector2D LinearSliderRange = FVector2D(Audio::ConvertToLinear(OutputRange.X), Audio::ConvertToLinear(OutputRange.Y));
+	return FMath::GetMappedRangeValueClamped(LinearSliderRange, NormalizedLinearSliderRange, LinearSliderValue);
+}
+
+const FText FFrequencyProcessor::GetUnitsText()
+{
+	return FText::FromString("Hz");;
+}
+
+const FVector2D FFrequencyProcessor::GetDefaultOutputRange()
+{
+	return FVector2D(MIN_FILTER_FREQUENCY, MAX_FILTER_FREQUENCY);
+}
+
+const float FFrequencyProcessor::GetOutputValue(const FVector2D OutputRange, const float InSliderValue)
+{
+	return	Audio::GetLogFrequencyClamped(InSliderValue, NormalizedLinearSliderRange, OutputRange);
+}
+
+const float FFrequencyProcessor::GetSliderValue(const FVector2D OutputRange, const float OutputValue)
+{
+	// edge case to avoid audio function returning negative value
+	if (FMath::IsNearlyEqual(OutputValue, OutputRange.X))
+	{
+		return NormalizedLinearSliderRange.X;
+	}
+	if (FMath::IsNearlyEqual(OutputValue, OutputRange.Y))
+	{
+		return NormalizedLinearSliderRange.Y;
+	}
+
+	return Audio::GetLinearFrequencyClamped(OutputValue, NormalizedLinearSliderRange, OutputRange);	
+}
diff --git a/Engine/Plugins/Runtime/AudioWidgets/Source/AudioWidgets/Public/AudioMaterialSlate/SAudioMaterialLabeledSlider.h b/Engine/Plugins/Runtime/AudioWidgets/Source/AudioWidgets/Public/AudioMaterialSlate/SAudioMaterialLabeledSlider.h
--- a/Engine/Plugins/Runtime/AudioWidgets/Source/AudioWidgets/Public/AudioMaterialSlate/SAudioMaterialLabeledSlider.h
+++ b/Engine/Plugins/Runtime/AudioWidgets/Source/AudioWidgets/Public/AudioMaterialSlate/SAudioMaterialLabeledSlider.h
@@ -3,6 +3,7 @@
 #pragma once
 
 #include "AudioMaterialSlate/AudioMaterialSlateTypes.h"
+#include "AudioWidgetsEnums.h"
 #include "Framework/SlateDelegates.h"
 #include "SAudioInputWidget.h"
 #include "Styling/ISlateStyle.h"
@@ -33,7 +34,13 @@ public:
 	SLATE_ATTRIBUTE(float, SliderValue)
 		
 	/** The slider's orientation. */
-	SLATE_ARGUMENT(EOrientation, Orientation)
+	SLATE_ARGUMENT(EOrientation, Orientation)	
+	
+	/** The slider's ValueType. */
+	SLATE_ARGUMENT(EAudioUnitsValueType, AudioUnitsValueType)
+	
+	/** Will the slider use Linear Output. This is used when ValueType is Volume */
+	SLATE_ARGUMENT(bool , bUseLinearOutput)
 	
 	/** When specified, use this as the slider's desired size */
 	SLATE_ATTRIBUTE(TOptional<FVector2D>, DesiredSizeOverride)
@@ -70,7 +77,7 @@ public:
 	virtual void SetDesiredSizeOverride(const FVector2D DesiredSize) override;
 
 	void SetOrientation(EOrientation InOrientation);
-	virtual void SetOutputRange(const FVector2D Range) override;
+	virtual void SetOutputRange(const FVector2D InRange) override;
 
 	// Text label functions 
 	void SetLabelBackgroundColor(FSlateColor InColor) override;
@@ -91,6 +98,9 @@ private:
 
 	// Holds the slider's orientation
 	TAttribute<EOrientation> Orientation;
+	
+	// Holds the slider's unit value type
+	TAttribute<EAudioUnitsValueType> AudioUnitsValueType;
 
 	// Optional override for desired size 
 	TAttribute<TOptional<FVector2D>> DesiredSizeOverride;
@@ -104,11 +114,13 @@ private:
 
 	// Range for output
 	FVector2D OutputRange = FVector2D(0.0f, 1.0f);
-	static const FVector2D NormalizedLinearSliderRange;
 
 	/** Switches between the vertical and horizontal views */
 	TSharedPtr<SWidgetSwitcher> LayoutWidgetSwitcher;
 
+	/**Hold the ref to the current Unit processor */
+	TSharedPtr<FAudioUnitProcessor> AudioUnitProcessor;
+
 private:
 
 	TSharedRef<SWidgetSwitcher> CreateWidgetLayout();
diff --git a/Engine/Plugins/Runtime/AudioWidgets/Source/AudioWidgets/Public/AudioWidgetsEnums.h b/Engine/Plugins/Runtime/AudioWidgets/Source/AudioWidgets/Public/AudioWidgetsEnums.h
--- a/Engine/Plugins/Runtime/AudioWidgets/Source/AudioWidgets/Public/AudioWidgetsEnums.h
+++ b/Engine/Plugins/Runtime/AudioWidgets/Source/AudioWidgets/Public/AudioWidgetsEnums.h
@@ -10,3 +10,11 @@ enum class EAudioPanelLayoutType : uint8
 	Basic    UMETA(DisplayName = "Basic"),
 	Advanced UMETA(DisplayName = "Advanced")
 };
+
+UENUM()
+enum class EAudioUnitsValueType : uint8
+{
+	Linear,
+	Frequency UMETA(DisplayName = "Frequency (Log)"),
+	Volume
+};
diff --git a/Engine/Plugins/Runtime/AudioWidgets/Source/AudioWidgets/Public/SAudioInputWidget.h b/Engine/Plugins/Runtime/AudioWidgets/Source/AudioWidgets/Public/SAudioInputWidget.h
--- a/Engine/Plugins/Runtime/AudioWidgets/Source/AudioWidgets/Public/SAudioInputWidget.h
+++ b/Engine/Plugins/Runtime/AudioWidgets/Source/AudioWidgets/Public/SAudioInputWidget.h
@@ -29,3 +29,81 @@ public:
 	virtual void SetShowUnitsText(const bool bShowUnitsText) = 0;
 	virtual void SetDesiredSizeOverride(const FVector2D DesiredSize) = 0;
 };
+
+/**
+*Processor to allow output and display of different Audio Units. This is because what is shown for the user and what the output is might not be the same. 
+*For example, when using the linear output option for volume, Volume is displayed as dB, but the output value will still be 0.0-1.0f.
+*/
+ struct FAudioUnitProcessor
+ {
+public:
+	
+	/**Get the units that the processor will output*/
+	virtual const FText GetUnitsText() { return FText(); };
+
+	/**Get the OutputValue. This is where to calculate what the processor will actually output*/
+	virtual const float GetOutputValue(const FVector2D OutputRange, const float InSliderValue);
+
+	/**Get the OutputValue that will be shown as text. This might be different then OutputValue depending on the type*/
+	virtual const float GetOutputValueForText(const FVector2D OutputRange, const float InSliderValue);
+
+	/**Get the Value for the slider. This should be clamped to be inside the given slider range*/
+	virtual const float GetSliderValue(const FVector2D OutputRange, const float OutputValue);
+
+	/**Get the Slider value for the text. option to override what would be returned when value is set directly to the text field. The value in the text field might differ from the slider value because the slider is usally just between 0-1*/
+	virtual const float GetSliderValueForText(const FVector2D OutputRange, const float OutputValue);
+	
+	/**Get the default Output range for the processor. This will be used when Slate is created, and default values are set.*/
+	virtual const FVector2D GetDefaultOutputRange() { return FVector2D(0.0f , 1.0f);};
+
+	/**Get the OutputRange. Possibility to modify the output range. Currently only used for Volume*/
+	virtual const FVector2D GetOutputRange(FVector2D InRange) { return InRange; };
+
+protected:
+
+	static const FVector2D NormalizedLinearSliderRange;
+
+ };
+
+ struct FVolumeProcessor : FAudioUnitProcessor
+ {
+
+	 FVolumeProcessor(bool bInUseLinearOutput){ bUseLinearOutput  = bInUseLinearOutput;}
+
+	 virtual const FText GetUnitsText() override;
+	 virtual const FVector2D GetDefaultOutputRange() override;
+
+	 virtual const float GetOutputValue(const FVector2D OutputRange, const float InSliderValue) override;
+	 virtual const float GetOutputValueForText(const FVector2D OutputRange, const float InSliderValue) override;
+	 virtual const float GetSliderValue(const FVector2D OutputRange, const float OutputValue) override;
+	 virtual const float GetSliderValueForText(const FVector2D OutputRange, const float OutputValue) override;
+	 virtual const FVector2D GetOutputRange(const FVector2D InRange) override;
+
+private:
+
+	/**
+	*Convert the given SliderValue to dB within the given range
+	*/
+	const float GetDbValueFromSliderValue(const FVector2D OutputRange, const float InSliderValue);
+
+	/**Convert decibels to linear 0-1 space*/
+	const float GetSliderValueFromDb(const FVector2D OutputRange, const float DbValue);
+
+private:
+
+	// Min/max possible values for output range, derived to avoid Audio::ConvertToLinear/dB functions returning NaN
+	static const float MinDbValue;
+	static const float MaxDbValue;
+
+	// Use linear (converted from dB, not normalized) output value. Only applies to the output value reported by GetOutputValue(); the text displayed will still be in decibels. 
+	bool bUseLinearOutput;
+ }; 
+ 
+ struct FFrequencyProcessor : FAudioUnitProcessor
+ {
+	 virtual const FText GetUnitsText() override;
+	 virtual const FVector2D GetDefaultOutputRange() override;
+
+	 virtual const float GetOutputValue(const FVector2D OutputRange, const float InSliderValue) override;
+	 virtual const float GetSliderValue(const FVector2D OutputRange, const float OutputValue) override;
+ };
diff --git a/Engine/Plugins/Runtime/Metasound/Source/MetasoundEditor/Private/MetasoundEditorGraphMemberDefaults.cpp b/Engine/Plugins/Runtime/Metasound/Source/MetasoundEditor/Private/MetasoundEditorGraphMemberDefaults.cpp
--- a/Engine/Plugins/Runtime/Metasound/Source/MetasoundEditor/Private/MetasoundEditorGraphMemberDefaults.cpp
+++ b/Engine/Plugins/Runtime/Metasound/Source/MetasoundEditor/Private/MetasoundEditorGraphMemberDefaults.cpp
@@ -184,13 +184,17 @@ void UMetasoundEditorGraphMemberDefaultFloat::PostEditChangeChainProperty(FPrope
 		SetDefault(Default);
 	}
 	else if (PropertyChangedEvent.GetPropertyName().IsEqual(GET_MEMBER_NAME_CHECKED(UMetasoundEditorGraphMemberDefaultFloat, WidgetType)) ||
-		PropertyChangedEvent.GetPropertyName().IsEqual(GET_MEMBER_NAME_CHECKED(UMetasoundEditorGraphMemberDefaultFloat, WidgetValueType)))
+		PropertyChangedEvent.GetPropertyName().IsEqual(GET_MEMBER_NAME_CHECKED(UMetasoundEditorGraphMemberDefaultFloat, WidgetUnitValueType)))
 	{
 		// Update VolumeWidgetDecibelRange based on current range (it might be stale)
-		if (WidgetValueType == EMetasoundMemberDefaultWidgetValueType::Volume && VolumeWidgetUseLinearOutput)
+		if (WidgetUnitValueType == EAudioUnitsValueType::Volume && VolumeWidgetUseLinearOutput)
 		{
 			VolumeWidgetDecibelRange = FVector2D(Audio::ConvertToDecibels(Range.X), Audio::ConvertToDecibels(Range.Y));
 		}
+		else if (WidgetUnitValueType == EAudioUnitsValueType::Frequency)
+		{
+			SetRange(FVector2D(MIN_FILTER_FREQUENCY, MAX_FILTER_FREQUENCY));
+		}
 		else
 		{
 			SetInitialRange();
@@ -204,7 +208,7 @@ void UMetasoundEditorGraphMemberDefaultFloat::PostEditChangeChainProperty(FPrope
 		if (PropertyChangedEvent.ChangeType != EPropertyChangeType::Interactive)
 		{
 			if (WidgetType != EMetasoundMemberDefaultWidget::None &&
-				WidgetValueType == EMetasoundMemberDefaultWidgetValueType::Volume)
+				WidgetUnitValueType == EAudioUnitsValueType::Volume)
 			{
 				if (VolumeWidgetUseLinearOutput)
 				{
diff --git a/Engine/Plugins/Runtime/Metasound/Source/MetasoundEditor/Private/MetasoundNodeDetailCustomization.cpp b/Engine/Plugins/Runtime/Metasound/Source/MetasoundEditor/Private/MetasoundNodeDetailCustomization.cpp
--- a/Engine/Plugins/Runtime/Metasound/Source/MetasoundEditor/Private/MetasoundNodeDetailCustomization.cpp
+++ b/Engine/Plugins/Runtime/Metasound/Source/MetasoundEditor/Private/MetasoundNodeDetailCustomization.cpp
@@ -347,8 +347,8 @@ namespace Metasound
 				IDetailCategoryBuilder& WidgetCategoryBuilder = InDetailLayout.EditCategory("EditorOptions");
 				DefaultRows.Add(WidgetCategoryBuilder.AddExternalObjectProperty(TArray<UObject*>({ DefaultFloat }), GET_MEMBER_NAME_CHECKED(UMetasoundEditorGraphMemberDefaultFloat, WidgetType)));
 				DefaultRows.Add(WidgetCategoryBuilder.AddExternalObjectProperty(TArray<UObject*>({ DefaultFloat }), GET_MEMBER_NAME_CHECKED(UMetasoundEditorGraphMemberDefaultFloat, WidgetOrientation)));
-				DefaultRows.Add(WidgetCategoryBuilder.AddExternalObjectProperty(TArray<UObject*>({ DefaultFloat }), GET_MEMBER_NAME_CHECKED(UMetasoundEditorGraphMemberDefaultFloat, WidgetValueType)));
-				if (DefaultFloat->WidgetType != EMetasoundMemberDefaultWidget::None && DefaultFloat->WidgetValueType == EMetasoundMemberDefaultWidgetValueType::Volume)
+				DefaultRows.Add(WidgetCategoryBuilder.AddExternalObjectProperty(TArray<UObject*>({ DefaultFloat }), GET_MEMBER_NAME_CHECKED(UMetasoundEditorGraphMemberDefaultFloat, WidgetUnitValueType)));
+				if (DefaultFloat->WidgetType != EMetasoundMemberDefaultWidget::None && DefaultFloat->WidgetUnitValueType == EAudioUnitsValueType::Volume)
 				{
 					DefaultRows.Add(WidgetCategoryBuilder.AddExternalObjectProperty(TArray<UObject*>({ DefaultFloat }), GET_MEMBER_NAME_CHECKED(UMetasoundEditorGraphMemberDefaultFloat, VolumeWidgetUseLinearOutput)));
 					if (DefaultFloat->VolumeWidgetUseLinearOutput)
diff --git a/Engine/Plugins/Runtime/Metasound/Source/MetasoundEditor/Private/SMetasoundGraphNode.cpp b/Engine/Plugins/Runtime/Metasound/Source/MetasoundEditor/Private/SMetasoundGraphNode.cpp
--- a/Engine/Plugins/Runtime/Metasound/Source/MetasoundEditor/Private/SMetasoundGraphNode.cpp
+++ b/Engine/Plugins/Runtime/Metasound/Source/MetasoundEditor/Private/SMetasoundGraphNode.cpp
@@ -4,6 +4,7 @@
 #include "AudioMaterialSlate/SAudioMaterialLabeledKnob.h"
 #include "AudioMaterialSlate/SAudioMaterialLabeledSlider.h"
 #include "AudioParameterControllerInterface.h"
+#include "AudioWidgetsEnums.h"
 #include "Components/AudioComponent.h"
 #include "GraphEditorSettings.h"
 #include "IAudioParameterTransmitter.h"
@@ -735,20 +736,21 @@ namespace Metasound
 							{
 								SAssignNew(InputWidget, SAudioMaterialLabeledSlider)
 									.Owner(GraphMember->GetOwningGraph())
+									.AudioUnitsValueType(DefaultFloat->WidgetUnitValueType)
+									.bUseLinearOutput(DefaultFloat->VolumeWidgetUseLinearOutput)
 									.OnValueChanged_Lambda(OnValueChangedLambda)
 									.OnValueCommitted_Lambda(OnValueCommittedLambda);
-								InputWidget->SetShowUnitsText(false);
 							}
 							else
 							{
 								// Create slider 
-								if (DefaultFloat->WidgetValueType == EMetasoundMemberDefaultWidgetValueType::Frequency)
+								if (DefaultFloat->WidgetUnitValueType == EAudioUnitsValueType::Frequency)
 								{
 									SAssignNew(InputWidget, SAudioFrequencySlider)
 										.OnValueChanged_Lambda(OnValueChangedLambda)
 										.OnValueCommitted_Lambda(OnValueCommittedLambda);
 								}
-								else if (DefaultFloat->WidgetValueType == EMetasoundMemberDefaultWidgetValueType::Volume)
+								else if (DefaultFloat->WidgetUnitValueType == EAudioUnitsValueType::Volume)
 								{
 									SAssignNew(InputWidget, SAudioVolumeSlider)
 										.OnValueChanged_Lambda(OnValueChangedLambda)
@@ -861,19 +863,18 @@ namespace Metasound
 									.OnValueChanged_Lambda(OnValueChangedLambda)
 									.OnMouseCaptureBegin_Lambda(OnRadialSliderMouseCaptureBeginLambda)
 									.OnMouseCaptureEnd_Lambda(OnRadialSliderMouseCaptureEndLambda);
-								InputWidget->SetShowUnitsText(false);								
 							}
 							else
 							{
 								// Create slider 
-								if (DefaultFloat->WidgetValueType == EMetasoundMemberDefaultWidgetValueType::Frequency)
+								if (DefaultFloat->WidgetUnitValueType == EAudioUnitsValueType::Frequency)
 								{
 									SAssignNew(InputWidget, SAudioFrequencyRadialSlider)
 										.OnValueChanged_Lambda(OnValueChangedLambda)
 										.OnMouseCaptureBegin_Lambda(OnRadialSliderMouseCaptureBeginLambda)
 										.OnMouseCaptureEnd_Lambda(OnRadialSliderMouseCaptureEndLambda);
 								}
-								else if (DefaultFloat->WidgetValueType == EMetasoundMemberDefaultWidgetValueType::Volume)
+								else if (DefaultFloat->WidgetUnitValueType == EAudioUnitsValueType::Volume)
 								{
 									SAssignNew(InputWidget, SAudioVolumeRadialSlider)
 										.OnValueChanged_Lambda(OnValueChangedLambda)
diff --git a/Engine/Plugins/Runtime/Metasound/Source/MetasoundEditor/Public/MetasoundEditorGraphMemberDefaults.h b/Engine/Plugins/Runtime/Metasound/Source/MetasoundEditor/Public/MetasoundEditorGraphMemberDefaults.h
--- a/Engine/Plugins/Runtime/Metasound/Source/MetasoundEditor/Public/MetasoundEditorGraphMemberDefaults.h
+++ b/Engine/Plugins/Runtime/Metasound/Source/MetasoundEditor/Public/MetasoundEditorGraphMemberDefaults.h
@@ -4,6 +4,7 @@
 #include "Algo/Count.h"
 #include "Algo/Transform.h"
 #include "AudioParameterControllerInterface.h"
+#include "AudioWidgetsEnums.h"
 #include "Components/Widget.h"
 #include "Delegates/Delegate.h"
 #include "EdGraph/EdGraphNode.h"
@@ -123,7 +124,7 @@ DECLARE_MULTICAST_DELEGATE_OneParam(FOnMetasoundInputClampDefaultChangedEvent, b
 
 
 UENUM()
-enum class EMetasoundMemberDefaultWidgetValueType : uint8
+enum class UE_DEPRECATED(5.5, "EMetasoundMemberDefaultWidgetValueType is deprecated, use EAudioUnitsValueType instead") EMetasoundMemberDefaultWidgetValueType : uint8
 {
 	Linear,
 	Frequency UMETA(DisplayName = "Frequency (Log)"),
@@ -154,8 +155,16 @@ public:
 	UPROPERTY(EditAnywhere, Category = Widget, meta=(DisplayName = "Orientation", EditCondition = "WidgetType == EMetasoundMemberDefaultWidget::Slider", EditConditionHides))
 	TEnumAsByte<EOrientation> WidgetOrientation = EOrientation::Orient_Horizontal;
 
-	UPROPERTY(EditAnywhere, Category = Widget, meta=(DisplayName = "Value Type", EditCondition = "WidgetType != EMetasoundMemberDefaultWidget::None", EditConditionHides))
-	EMetasoundMemberDefaultWidgetValueType WidgetValueType = EMetasoundMemberDefaultWidgetValueType::Linear;
+#if WITH_EDITORONLY_DATA
+	PRAGMA_DISABLE_DEPRECATION_WARNINGS
+	UE_DEPRECATED(5.5, "WidgetValueType has been deprecated. Use WidgetUnitValueType instead.")
+	UPROPERTY(meta = (DeprecatedProperty, DeprecationMessage = "WidgetValueType has been deprecated. Use WidgetUnitValueType instead."))
+	EMetasoundMemberDefaultWidgetValueType WidgetValueType_DEPRECATED = EMetasoundMemberDefaultWidgetValueType::Linear;
+	PRAGMA_ENABLE_DEPRECATION_WARNINGS
+#endif//WITH_EDITORONLY_DATA
+
+	UPROPERTY(EditAnywhere, Category = Widget, meta = (DisplayName = "Value Type", EditCondition = "WidgetType != EMetasoundMemberDefaultWidget::None", EditConditionHides))
+	EAudioUnitsValueType WidgetUnitValueType = EAudioUnitsValueType::Linear;
 
 	/** If true, output linear value. Otherwise, output dB value. The volume widget itself will always display the value in dB. The Default Value and Range are linear. */
 	UPROPERTY(EditAnywhere, Category = Widget, meta = (DisplayName = "Output Linear"))
