AudioMaterialMeter Value as ChannelInfo & Value ScaleHash
#rnx

diff --git a/Engine/Plugins/Runtime/AudioWidgets/Source/AudioWidgets/Private/AudioMaterialSlate/AudioMaterialMeter.cpp b/Engine/Plugins/Runtime/AudioWidgets/Source/AudioWidgets/Private/AudioMaterialSlate/AudioMaterialMeter.cpp
--- a/Engine/Plugins/Runtime/AudioWidgets/Source/AudioWidgets/Private/AudioMaterialSlate/AudioMaterialMeter.cpp
+++ b/Engine/Plugins/Runtime/AudioWidgets/Source/AudioWidgets/Private/AudioMaterialSlate/AudioMaterialMeter.cpp
@@ -2,18 +2,30 @@
 
 
 #include "AudioMaterialSlate/AudioMaterialMeter.h"
+#include "AudioBusSubsystem.h"
 #include "AudioMaterialSlate/AudioMaterialSlateTypes.h"
 #include "AudioMaterialSlate/SAudioMaterialMeter.h"
+#include "AudioMixerDevice.h"
 #include "AudioWidgetsStyle.h"
-#include "Widgets/SWeakWidget.h"
 #include "Components/AudioComponent.h"
+#include "SAudioMeter.h"
+#include "UObject/UObjectGlobals.h"
+#include "Widgets/SWeakWidget.h"
 
 #define LOCTEXT_NAMESPACE "AudioWidgets"
 
 UAudioMaterialMeter::UAudioMaterialMeter()
 {
+	Orientation = EOrientation::Orient_Vertical;
+
 	//get default style
 	WidgetStyle = FAudioWidgetsStyle::Get().GetWidgetStyle<FAudioMaterialMeterStyle>("AudioMaterialMeter.Style");
+
+	// Add a single channel as a default just so it can be seen when somebody makes one
+	FMeterChannelInfo DefaultInfo;
+	DefaultInfo.MeterValue = -6.0f;
+	DefaultInfo.PeakValue = -3.0f;
+	MeterChannelInfo.Add(DefaultInfo);
 }
 
 #if WITH_EDITOR
@@ -33,7 +45,10 @@ void UAudioMaterialMeter::SynchronizeProperties()
 	}
 
 	Meter->ApplyNewMaterial();
-	Meter->SetValue(MeterValue);
+	Meter->SetOrientation(Orientation);
+
+	TAttribute<TArray<FMeterChannelInfo>> MeterChannelInfoBinding = PROPERTY_BINDING(TArray<FMeterChannelInfo>, MeterChannelInfo);
+	Meter->SetMeterChannelInfo(MeterChannelInfoBinding);
 }
 
 void UAudioMaterialMeter::ReleaseSlateResources(bool bReleaseChildren)
@@ -41,22 +56,20 @@ void UAudioMaterialMeter::ReleaseSlateResources(bool bReleaseChildren)
 	Meter.Reset();
 }
 
-float UAudioMaterialMeter::GetValue() const
-{
-	return MeterValue;
-}
-
-void UAudioMaterialMeter::SetValue(float InValue)
+TArray<FMeterChannelInfo> UAudioMaterialMeter::GetMeterChannelInfo() const
 {
 	if (Meter.IsValid())
 	{
-		Meter->SetValue(InValue);
+		return Meter->GetMeterChannelInfo();
 	}
+	return TArray<FMeterChannelInfo>();
+}
 
-	if (MeterValue != InValue)
+void UAudioMaterialMeter::SetMeterChannelInfo(const TArray<FMeterChannelInfo>& InMeterChannelInfo)
+{
+	if (Meter.IsValid())
 	{
-		MeterValue = InValue;
-		HandleOnValueChanged(InValue);
+		Meter->SetMeterChannelInfo(InMeterChannelInfo);
 	}
 }
 
@@ -64,16 +77,9 @@ TSharedRef<SWidget> UAudioMaterialMeter::RebuildWidget()
 {
 	Meter = SNew(SAudioMaterialMeter)
 		.Owner(this)
-		.AudioMaterialMeterStyle(&WidgetStyle)
-		.ValueAttribute(MeterValue)
-		.OnValueChanged(BIND_UOBJECT_DELEGATE(FOnFloatValueChanged, HandleOnValueChanged));
+		.AudioMaterialMeterStyle(&WidgetStyle);
 
 	return Meter.ToSharedRef();
 }
 
-void UAudioMaterialMeter::HandleOnValueChanged(float InValue)
-{
-	OnValueChanged.Broadcast(InValue);
-}
-
 #undef LOCTEXT_NAMESPACE
diff --git a/Engine/Plugins/Runtime/AudioWidgets/Source/AudioWidgets/Private/AudioMaterialSlate/AudioMaterialSlateTypes.cpp b/Engine/Plugins/Runtime/AudioWidgets/Source/AudioWidgets/Private/AudioMaterialSlate/AudioMaterialSlateTypes.cpp
--- a/Engine/Plugins/Runtime/AudioWidgets/Source/AudioWidgets/Private/AudioMaterialSlate/AudioMaterialSlateTypes.cpp
+++ b/Engine/Plugins/Runtime/AudioWidgets/Source/AudioWidgets/Private/AudioMaterialSlate/AudioMaterialSlateTypes.cpp
@@ -2,6 +2,7 @@
 
 
 #include "AudioMaterialSlate/AudioMaterialSlateTypes.h"
+#include "Styling/StyleDefaults.h"
 
 FAudioMaterialWidgetStyle::FAudioMaterialWidgetStyle()
 	:DesiredSize(32.f, 32.f)
@@ -94,12 +95,21 @@ void FAudioMaterialKnobStyle::GetResources(TArray<const FSlateBrush*>& OutBrushe
 }
 
 FAudioMaterialMeterStyle::FAudioMaterialMeterStyle()
-	:MeterFillMinColor(FLinearColor::White)
+	: MeterFillMinColor(FLinearColor::White)
 	, MeterFillMidColor(FLinearColor::White)
 	, MeterFillMaxColor(FLinearColor::White)
 	, MeterOffFillColor(FLinearColor::Black)
-{
-	DesiredSize = FVector2f(32.f, 512.f);
+	, MeterPadding(FVector2D(10.0f, 5.0f))
+	, ValueRangeDb(FVector2D(-60, 10))
+	, bShowScale(true)
+	, bScaleSide(true)
+	, ScaleHashOffset(5.0f)
+	, ScaleHashWidth(10.0f)
+	, ScaleHashHeight(1.0f)
+	, DecibelsPerHash(5)
+	, Font(FStyleDefaults::GetFontInfo(5))
+{
+	DesiredSize = FVector2f(25.f, 512.f);
 }
 
 const FName FAudioMaterialMeterStyle::TypeName(TEXT("FAudioMaterialMeterStyle"));
diff --git a/Engine/Plugins/Runtime/AudioWidgets/Source/AudioWidgets/Private/AudioMaterialSlate/SAudioMaterialMeter.cpp b/Engine/Plugins/Runtime/AudioWidgets/Source/AudioWidgets/Private/AudioMaterialSlate/SAudioMaterialMeter.cpp
--- a/Engine/Plugins/Runtime/AudioWidgets/Source/AudioWidgets/Private/AudioMaterialSlate/SAudioMaterialMeter.cpp
+++ b/Engine/Plugins/Runtime/AudioWidgets/Source/AudioWidgets/Private/AudioMaterialSlate/SAudioMaterialMeter.cpp
@@ -3,9 +3,13 @@
 
 #include "AudioMaterialSlate/SAudioMaterialMeter.h"
 #include "AudioMaterialSlate/AudioMaterialMeter.h"
+#include "AudioWidgetsStyle.h"
 #include "Components/AudioComponent.h"
+#include "Fonts/FontMeasure.h"
+#include "Framework/Application/SlateApplication.h"
 #include "SlateOptMacros.h"
 #include "Styling/SlateBrush.h"
+#include "Styling/StyleDefaults.h"
 
 
 BEGIN_SLATE_FUNCTION_BUILD_OPTIMIZATION
@@ -13,9 +17,15 @@ BEGIN_SLATE_FUNCTION_BUILD_OPTIMIZATION
 void SAudioMaterialMeter::Construct(const FArguments& InArgs)
 {
 	Owner = InArgs._Owner;
-	ValueAttribute = InArgs._ValueAttribute;
-	OnValueChanged = InArgs._OnValueChanged;
-	AudioMaterialMeterStyle = InArgs._AudioMaterialMeterStyle;
+	Orientation = InArgs._Orientation;
+
+	Style = InArgs._AudioMaterialMeterStyle;
+	MeterChannelInfoAttribute = InArgs._MeterChannelInfo;
+
+	if (!Cast<UWidget>(Owner))
+	{
+		Style = &FAudioWidgetsStyle::Get().GetWidgetStyle<FAudioMaterialMeterStyle>("AudioMaterialMeter.Style");
+	}
 
 	ApplyNewMaterial();
 }
@@ -24,33 +34,230 @@ END_SLATE_FUNCTION_BUILD_OPTIMIZATION
 
 int32 SAudioMaterialMeter::OnPaint(const FPaintArgs& Args, const FGeometry& AllottedGeometry, const FSlateRect& MyCullingRect, FSlateWindowElementList& OutDrawElements, int32 LayerId, const FWidgetStyle& InWidgetStyle, bool bParentEnabled) const
 {
-	if (AudioMaterialMeterStyle)
+	if (Style && !DynamicMaterials.IsEmpty())
 	{
-		UMaterialInstanceDynamic* DynamicMaterial = AudioMaterialMeterStyle->GetDynamicMaterial();
-		if (IsValid(DynamicMaterial))
+		const float AllottedWidth = Orientation == Orient_Vertical ? AllottedGeometry.GetLocalSize().X : AllottedGeometry.GetLocalSize().Y;
+		const float AllottedHeight = Orientation == Orient_Vertical ? AllottedGeometry.GetLocalSize().Y : AllottedGeometry.GetLocalSize().X;
+
+		FGeometry MeterGeometry = AllottedGeometry;
+
+		// Get the scale hash offset
+		float ScaleOffset = 0.0f;
+		if (Style->bShowScale && Style->bScaleSide)
 		{
+			ScaleOffset = GetScaleWidth();
+		}
+
+		TArray<FMeterChannelInfo> ChannelInfos = MeterChannelInfoAttribute.Get();
+		int32 NumChannels = ChannelInfos.Num();
+
+		const FLinearColor FinalColorAndOpacity(InWidgetStyle.GetColorAndOpacityTint());
+
+		//Draw Meter for every channel
+		for (int32 ChannelIndex = 0; ChannelIndex < NumChannels; ++ChannelIndex)
+		{	
+			TWeakObjectPtr<UMaterialInstanceDynamic> DynamicMaterial = DynamicMaterials[ChannelIndex];
+
+			if (DynamicMaterial.IsValid())
+			{
+				DynamicMaterial->SetVectorParameterValue(FName("A (V3)"), Style->MeterFillMinColor);
+				DynamicMaterial->SetVectorParameterValue(FName("B (V3)"), Style->MeterFillMidColor);
+				DynamicMaterial->SetVectorParameterValue(FName("C (V3)"), Style->MeterFillMaxColor);
+				DynamicMaterial->SetVectorParameterValue(FName("OffColor"), Style->MeterOffFillColor);
+				DynamicMaterial->SetVectorParameterValue(FName("DotsOffColor"), Style->MeterOffFillColor);
 
-			DynamicMaterial->SetVectorParameterValue(FName("A (V3)"), AudioMaterialMeterStyle->MeterFillMinColor);
-			DynamicMaterial->SetVectorParameterValue(FName("B (V3)"), AudioMaterialMeterStyle->MeterFillMidColor);
-			DynamicMaterial->SetVectorParameterValue(FName("C (V3)"), AudioMaterialMeterStyle->MeterFillMaxColor);
-			DynamicMaterial->SetVectorParameterValue(FName("OffColor"), AudioMaterialMeterStyle->MeterOffFillColor);
-			DynamicMaterial->SetVectorParameterValue(FName("DotsOffColor"), AudioMaterialMeterStyle->MeterOffFillColor);
+				float ChannelMeterValueDb = FMath::GetMappedRangeValueClamped(Style->ValueRangeDb,FVector2D(0.f,1.f),  ChannelInfos[ChannelIndex].MeterValue);
+				DynamicMaterial->SetScalarParameterValue(FName("VALUE"), ChannelMeterValueDb);
 
-			const float Value = ValueAttribute.Get();
-			DynamicMaterial->SetScalarParameterValue(FName("VALUE"), FMath::Clamp(Value, 0.f, 1.f));
+				DynamicMaterial->SetScalarParameterValue(FName("LocalWidth"), AllottedGeometry.GetLocalSize().X);
+				DynamicMaterial->SetScalarParameterValue(FName("LocalHeigth"), AllottedGeometry.GetLocalSize().Y);			
+
+				const bool bEnabled = ShouldBeEnabled(bParentEnabled);
+				const ESlateDrawEffect DrawEffects = bEnabled ? ESlateDrawEffect::None : ESlateDrawEffect::DisabledEffect;
+			
+				FSlateBrush Brush;
+				Brush.SetResourceObject(DynamicMaterial.Get());
 
-			DynamicMaterial->SetScalarParameterValue(FName("LocalWidth"), AllottedGeometry.GetLocalSize().X);
-			DynamicMaterial->SetScalarParameterValue(FName("LocalHeigth"), AllottedGeometry.GetLocalSize().Y);			
+				// rotate the meter 90deg if should be horizontal.
+				if (Orientation == Orient_Horizontal)
+				{
+					//rotate
+					FSlateRenderTransform SlateRenderTransform = TransformCast<FSlateRenderTransform>(Concatenate(Inverse(FVector2D(0, AllottedHeight)), FQuat2D(FMath::DegreesToRadians(90.0f))));
+					// create a child geometry matching this one, but with the render transform that will be passed to the drawed Meter.
+					MeterGeometry = AllottedGeometry.MakeChild(
+						FVector2D(AllottedWidth, AllottedHeight),
+						FSlateLayoutTransform(),
+						SlateRenderTransform, FVector2D::ZeroVector);
+				}
+	
+				FVector2D MeterTopLeft = FVector2D(Style->MeterPadding.X + ScaleOffset + ChannelIndex * (Style->DesiredSize.X + Style->MeterPadding.X), Style->MeterPadding.Y);
+				FVector2D MeterSize = FVector2D(Style->DesiredSize.X, AllottedHeight - Style->MeterPadding.Y);
+
+				FSlateDrawElement::MakeBox(
+					OutDrawElements,
+					LayerId,
+					MeterGeometry.ToPaintGeometry(MeterSize, FSlateLayoutTransform(MeterTopLeft)),	
+					&Brush,
+					DrawEffects,
+					FinalColorAndOpacity);
+			}
+			else
+			{
+				if (Style)
+				{
+					DynamicMaterials[ChannelIndex] = Style->CreateDynamicMaterial(Owner.Get());
+				}
+			}
 		}
 
-		const bool bEnabled = ShouldBeEnabled(bParentEnabled);
-		const ESlateDrawEffect DrawEffects = bEnabled ? ESlateDrawEffect::None : ESlateDrawEffect::DisabledEffect;
+		// Draw the scale hash
+		if (Style->bShowScale)
+		{
+
+			int32 DecibelsPerHash = Style->DecibelsPerHash;
+			int32 MinValueDb = FMath::Min((int32)Style->ValueRangeDb.X, (int32)Style->ValueRangeDb.Y);
+			int32 MaxValueDb = FMath::Max((int32)Style->ValueRangeDb.Y, (int32)Style->ValueRangeDb.Y);
+
+			// Snap the min/max values to the nearest hash 
+			MaxValueDb -= MaxValueDb % DecibelsPerHash;
+			MinValueDb -= MinValueDb % DecibelsPerHash;
+
+			float ScaleHashHalfHeight = 0.5f * Style->ScaleHashHeight;
+			FVector2D HashSize = FVector2D(Style->ScaleHashWidth, Style->ScaleHashHeight);
+
+			const TSharedRef<FSlateFontMeasure> FontMeasureService = FSlateApplication::Get().GetRenderer()->GetFontMeasureService();
+
+			// Measure the min value label size so we can right-justify if needed
+			FVector2D MinValueLabelSize = FontMeasureService->Measure(FString::FromInt(MinValueDb), Style->Font);
+
+			// Get the size of the negative sign to use to offset the label text in horizontal mode
+			FVector2D NegativeSignSize;
+			if (Orientation == Orient_Horizontal)
+			{
+				NegativeSignSize = FontMeasureService->Measure(TEXT("-"), Style->Font);
+			}
+
+			int32 ValueDelta = MinValueDb - MaxValueDb;
+			int32 CurrentHashValue = MaxValueDb;
+			while (CurrentHashValue >= MinValueDb)
+			{
+				// Get the fractional value for hash mark
+				const float CurrentHashMeterValuePercent = FMath::Clamp(((float)CurrentHashValue - (float)MaxValueDb) / ValueDelta, 0.0f, 1.0f);
 			
-		const FLinearColor FinalColorAndOpacity(InWidgetStyle.GetColorAndOpacityTint());
+				// Get the Y location for the hash
+				float HashPixelCenter = CurrentHashMeterValuePercent * (AllottedHeight - 2.0f * Style->MeterPadding.Y);
+
+				FVector2D HashTopLeft;
+				HashTopLeft.Y = Style->MeterPadding.Y + HashPixelCenter - ScaleHashHalfHeight;
+				if (Style->bScaleSide)
+				{
+					HashTopLeft.X = Style->MeterPadding.X + ScaleOffset - Style->ScaleHashOffset - Style->ScaleHashWidth;
+				}
+				else
+				{
+					HashTopLeft.X = (Style->DesiredSize.X + Style->MeterPadding.X) * NumChannels + Style->ScaleHashOffset;
+				}
+
+				// Draw hash
+				FSlateBrush HashBrush;
+				FSlateDrawElement::MakeBox(
+					OutDrawElements,
+					LayerId,
+					MeterGeometry.ToPaintGeometry(HashSize, FSlateLayoutTransform(HashTopLeft)),
+					&HashBrush,
+					ESlateDrawEffect::None,
+					FinalColorAndOpacity);
+
+				FString LabelString = FString::FromInt(CurrentHashValue);
+
+				bool bIsNegative = (CurrentHashValue < 0);
+
+				FVector2D LabelSize;
+
+				if (Orientation == Orient_Horizontal && bIsNegative)
+				{
+					// We want to center the text on just the positive portion of the number
+					FString LabelStringPositive = FString::FromInt(FMath::Abs(CurrentHashValue));
+					LabelSize = FontMeasureService->Measure(LabelStringPositive, Style->Font);
+				}
+				else
+				{
+					LabelSize = FontMeasureService->Measure(LabelString, Style->Font);
+				}
+
+				FText LabelText = FText::FromString(LabelString);
+
+				FGeometry TextGeometry = MeterGeometry;
+				FVector2D LabelTopLeft;
+
+				if (Orientation == Orient_Horizontal)
+				{
+					//Calculate offset from TopLeft
+					LabelTopLeft.Y = Style->MeterPadding.Y + HashPixelCenter + 0.5f * LabelSize.X;
+
+					if (bIsNegative)
+					{
+						LabelTopLeft.Y += NegativeSignSize.X;
+					}
+
+					if (Style->bScaleSide)
+					{
+						LabelTopLeft.X = Style->MeterPadding.X - 2.0f + (MinValueLabelSize.Y - LabelSize.Y);
+					}
+					else
+					{
+						LabelTopLeft.X = (Style->DesiredSize.X + Style->MeterPadding.X) * NumChannels + Style->ScaleHashOffset + Style->ScaleHashWidth + 2.0f;
+					}
 
-		FSlateBrush Brush;
-		Brush.SetResourceObject(DynamicMaterial);
-		FSlateDrawElement::MakeBox(OutDrawElements, LayerId++, AllottedGeometry.ToPaintGeometry(), &Brush, DrawEffects, FinalColorAndOpacity);
+					// Undo the rotation for horizontal before we do the rendering of the scale value 
+					FSlateRenderTransform RotationTransform = FSlateRenderTransform(FQuat2D(FMath::DegreesToRadians(-90.0f)));
+
+					FSlateLayoutTransform ChildLayoutTransform(1.0f, TransformPoint(1.0f, LabelTopLeft));
+
+					TextGeometry = MeterGeometry.MakeChild(
+						LabelSize,
+						ChildLayoutTransform,
+						RotationTransform,
+						FVector2D(0.0f, 0.0f));
+				}
+				else//Vertical
+				{
+					//Calculate offset from TopLeft
+					LabelTopLeft.Y = Style->MeterPadding.Y + HashPixelCenter - 0.5f * LabelSize.Y;
+
+					if (Style->bScaleSide)
+					{
+						LabelTopLeft.X = Style->MeterPadding.X + ScaleOffset - 2.0f - LabelSize.X - Style->ScaleHashOffset - Style->ScaleHashWidth;
+					}
+					else
+					{
+						LabelTopLeft.X = (Style->DesiredSize.X + Style->MeterPadding.X) * NumChannels + Style->ScaleHashOffset + Style->ScaleHashWidth + 2.0f;
+					}
+
+					FSlateRenderTransform RotationTransform = FSlateRenderTransform(FQuat2D(FMath::DegreesToRadians(0.0f)));
+
+					FSlateLayoutTransform ChildLayoutTransform(1.0f, TransformPoint(1.0f, LabelTopLeft));
+
+					TextGeometry = MeterGeometry.MakeChild(
+						LabelSize,
+						ChildLayoutTransform,
+						RotationTransform,
+						FVector2D(0.0f, 0.0f));
+				}
+
+				// Draw text label
+				FSlateDrawElement::MakeText(
+					OutDrawElements,
+					LayerId,
+					TextGeometry.ToPaintGeometry(),
+					LabelText,
+					Style->Font,
+					ESlateDrawEffect::None,
+					FinalColorAndOpacity);
+
+				CurrentHashValue -= Style->DecibelsPerHash;
+			}
+		}
 	}
 
 	return LayerId;
@@ -58,34 +265,84 @@ int32 SAudioMaterialMeter::OnPaint(const FPaintArgs& Args, const FGeometry& Allo
 
 FVector2D SAudioMaterialMeter::ComputeDesiredSize(float) const
 {
-	if (AudioMaterialMeterStyle)
+	static const FVector2D DefaultMeterDesiredSize(50.0f, 50.0f);
+
+	if (Style == nullptr)
 	{
-		return FVector2D(AudioMaterialMeterStyle->DesiredSize);
+		return DefaultMeterDesiredSize;
 	}
 
-	return FVector2D::ZeroVector;
+	TArray<FMeterChannelInfo> ChannelInfo = MeterChannelInfoAttribute.Get();
+	int32 NumChannels = FMath::Max(ChannelInfo.Num(), 1);
+
+	FVector2D Size = FVector2D((Style->DesiredSize.X + Style->MeterPadding.X) * NumChannels, Style->DesiredSize.Y+ Style->MeterPadding.Y);
+
+	// remove the end padding
+	Size.X += Style->MeterPadding.X;
+	Size.Y += Style->MeterPadding.Y;
+
+	// Add the width for the scale if it's been set to show
+	if (Style->bShowScale)
+	{
+		Size.X += GetScaleWidth();
+	}
+
+	if (Orientation == Orient_Horizontal)
+	{
+		return FVector2D(Size.Y, Size.X);
+	}
+
+	return Size;
 }
 
-void SAudioMaterialMeter::SetValue(float InValueAttribute)
+void SAudioMaterialMeter::SetOrientation(EOrientation InOrientation)
 {
-	CommitValue(InValueAttribute);
+	if (Orientation != InOrientation)
+	{
+		Orientation = InOrientation;
+		Invalidate(EInvalidateWidgetReason::Layout);
+	}
 }
 
-void SAudioMaterialMeter::ApplyNewMaterial()
+TArray<TWeakObjectPtr<UMaterialInstanceDynamic>> SAudioMaterialMeter::ApplyNewMaterial()
 {
-	if (AudioMaterialMeterStyle)
+	int32 NumChannels = MeterChannelInfoAttribute.Get().Num();
+
+	DynamicMaterials.SetNum(NumChannels);
+	DynamicMaterials.Empty();
+
+	if (Style)
 	{
-		AudioMaterialMeterStyle->CreateDynamicMaterial(Owner.Get());
+		for (int32 ChannelIndex = 0; ChannelIndex < NumChannels; ++ChannelIndex)
+		{
+			UMaterialInstanceDynamic* MaterialInstance = Style->CreateDynamicMaterial(Owner.Get());
+			DynamicMaterials.Add(MaterialInstance);
+		}
 	}
+
+	return DynamicMaterials;
 }
 
-void SAudioMaterialMeter::CommitValue(float NewValue)
+void SAudioMaterialMeter::SetMeterChannelInfo(const TAttribute<TArray<FMeterChannelInfo>>& InMeterChannelInfo)
 {
-	float Val = FMath::Clamp(NewValue, 0.f , 1.f);
-	if (!ValueAttribute.IsBound())
+	SetAttribute(MeterChannelInfoAttribute, InMeterChannelInfo, EInvalidateWidgetReason::Paint);
+	ApplyNewMaterial();
+}
+
+TArray<FMeterChannelInfo> SAudioMaterialMeter::GetMeterChannelInfo() const
+{
+	return MeterChannelInfoAttribute.Get();
+}
+
+float SAudioMaterialMeter::GetScaleWidth() const
+{
+	const TSharedRef<FSlateFontMeasure> FontMeasureService = FSlateApplication::Get().GetRenderer()->GetFontMeasureService();
+	FVector2D LabelSize = FontMeasureService->Measure(FString::FromInt(-60), Style->Font);
+
+	float ScaleWidth = Style->ScaleHashWidth + Style->ScaleHashOffset;
+	if (Orientation == Orient_Horizontal)
 	{
-		ValueAttribute.Set(Val);
+		return ScaleWidth + LabelSize.Y;
 	}
-
-	OnValueChanged.ExecuteIfBound(Val);
+	return ScaleWidth + LabelSize.X;
 }
diff --git a/Engine/Plugins/Runtime/AudioWidgets/Source/AudioWidgets/Public/AudioMaterialSlate/AudioMaterialMeter.h b/Engine/Plugins/Runtime/AudioWidgets/Source/AudioWidgets/Public/AudioMaterialSlate/AudioMaterialMeter.h
--- a/Engine/Plugins/Runtime/AudioWidgets/Source/AudioWidgets/Public/AudioMaterialSlate/AudioMaterialMeter.h
+++ b/Engine/Plugins/Runtime/AudioWidgets/Source/AudioWidgets/Public/AudioMaterialSlate/AudioMaterialMeter.h
@@ -2,15 +2,16 @@
 
 #pragma once
 
+#include "AudioDefines.h"
 #include "AudioMaterialSlate/AudioMaterialSlateTypes.h"
+#include "AudioMeterStyle.h"
+#include "AudioMeterTypes.h"
 #include "Components/Widget.h"
 #include "Delegates/Delegate.h"
 #include "AudioMaterialMeter.generated.h"
 
 class SAudioMaterialMeter;
 
-DECLARE_DYNAMIC_MULTICAST_DELEGATE_OneParam(FOnMeterValueChangedEvent, float, Value);
-
 /**
  * Meter is rendered by using material instead of texture.
  *
@@ -29,6 +30,10 @@ public:
 	UPROPERTY(EditAnywhere, BlueprintReadWrite, Category = "Style", meta = (DisplayName = "Style", ShowOnlyInnerProperties))
 	FAudioMaterialMeterStyle WidgetStyle;
 
+	/** The Meter's orientation. */
+	UPROPERTY(EditAnywhere, BlueprintReadOnly, Category = Appearance)
+	TEnumAsByte<EOrientation> Orientation;
+
 public:
 
 #if WITH_EDITOR
@@ -43,19 +48,21 @@ public:
 	virtual void ReleaseSlateResources(bool bReleaseChildren) override;
 	// End of UVisual
 
-	/** Get the current value of the meter.*/
-	UFUNCTION(BlueprintPure, Category = "Behavior")
-	float GetValue() const;
+	/** Gets the current linear values of the meter. */
+	UFUNCTION(BlueprintCallable, Category = "Behavior")
+	TArray<FMeterChannelInfo> GetMeterChannelInfo() const;
 
-	/** Set the current value of the Meter.*/
+	/** Sets the current meter values. */
 	UFUNCTION(BlueprintCallable, Category = "Behavior")
-	void SetValue(float InValue);
+	void SetMeterChannelInfo(const TArray<FMeterChannelInfo>& InMeterChannelInfo);
 
 public:
 
-	/** Called when the meter value changes. */
-	UPROPERTY(BlueprintAssignable, Category = "Widget Event")
-	FOnMeterValueChangedEvent OnValueChanged;
+	DECLARE_DYNAMIC_DELEGATE_RetVal(TArray<FMeterChannelInfo>, FGetMeterChannelInfo);
+
+	///** A bindable delegate to allow logic to drive the value of the meter */
+	UPROPERTY()
+	FGetMeterChannelInfo MeterChannelInfoDelegate;	
 
 protected:
 
@@ -63,15 +70,15 @@ protected:
 	virtual TSharedRef<SWidget> RebuildWidget() override;
 	// End of UWidget
 
-	void HandleOnValueChanged(float InValue);
+	PROPERTY_BINDING_IMPLEMENTATION(TArray<FMeterChannelInfo>, MeterChannelInfo);
 
 private:
 
-	/**Current Value of the Meter*/
-	UPROPERTY(EditAnywhere, BlueprintSetter = SetValue, BlueprintGetter = GetValue, Category = "Appearance", meta = (UIMin = "0", UIMax = "1"))
-	float MeterValue = 1.f;
-
 	/** Native Slate Widget */
 	TSharedPtr<SAudioMaterialMeter> Meter;
 
+	/** The current meter value to display. */
+	UPROPERTY(EditAnywhere, BlueprintSetter = SetMeterChannelInfo, BlueprintGetter = GetMeterChannelInfo, Category = MeterValues)
+	TArray<FMeterChannelInfo> MeterChannelInfo;
+
 };
diff --git a/Engine/Plugins/Runtime/AudioWidgets/Source/AudioWidgets/Public/AudioMaterialSlate/AudioMaterialSlateTypes.h b/Engine/Plugins/Runtime/AudioWidgets/Source/AudioWidgets/Public/AudioMaterialSlate/AudioMaterialSlateTypes.h
--- a/Engine/Plugins/Runtime/AudioWidgets/Source/AudioWidgets/Public/AudioMaterialSlate/AudioMaterialSlateTypes.h
+++ b/Engine/Plugins/Runtime/AudioWidgets/Source/AudioWidgets/Public/AudioMaterialSlate/AudioMaterialSlateTypes.h
@@ -224,6 +224,63 @@ struct AUDIOWIDGETS_API FAudioMaterialMeterStyle : public FAudioMaterialWidgetSt
 	FLinearColor MeterOffFillColor;
 	FAudioMaterialMeterStyle& SetMeterOffFillColor(const FLinearColor& InColor) { MeterOffFillColor = InColor; return *this; }
 
+	// How much padding to add around the meter
+	UPROPERTY(EditAnywhere, BlueprintReadWrite, Category = Appearance)
+	FVector2D MeterPadding;
+	FAudioMaterialMeterStyle& SetMeterpadding(const FVector2D InPadding) { MeterPadding = InPadding; return *this; }
+
+	// The minimum and maximum value to display in dB (values are clamped in this range)
+	UPROPERTY(EditAnywhere, BlueprintReadWrite, Category = Appearance)
+	FVector2D ValueRangeDb;
+	FAudioMaterialMeterStyle& SetValueRangeDb(const FVector2D& InValueRangeDb) { ValueRangeDb = InValueRangeDb; return *this; }
+
+	// Whether or not to show the decibel scale alongside the meter
+	UPROPERTY(EditAnywhere, BlueprintReadWrite, Category = Appearance)
+	bool bShowScale;
+	FAudioMaterialMeterStyle& SetShowScale(bool bInShowScale) { bShowScale = bInShowScale; return *this; }
+
+	// Which side to show the scale. If vertical, true means left side, false means right side. If horizontal, true means above, false means below.
+	UPROPERTY(EditAnywhere, BlueprintReadWrite, Category = Appearance)
+	bool bScaleSide;
+	FAudioMaterialMeterStyle& SetScaleSide(bool bInScaleSide) { bScaleSide = bInScaleSide; return *this; }
+
+	// Offset for the hashes
+	UPROPERTY(EditAnywhere, BlueprintReadWrite, Category = Appearance)
+	float ScaleHashOffset;
+	FAudioMaterialMeterStyle& SetScaleHashOffset(float InScaleHashOffset) { ScaleHashOffset = InScaleHashOffset; return *this; }
+
+	// The width of each hash mark
+	UPROPERTY(EditAnywhere, BlueprintReadWrite, Category = Appearance)
+	float ScaleHashWidth;
+	FAudioMaterialMeterStyle& SetScaleHashWidth(float InScaleHashWidth) { ScaleHashWidth = InScaleHashWidth; return *this; }
+
+	// The height of each hash mark
+	UPROPERTY(EditAnywhere, BlueprintReadWrite, Category = Appearance)
+	float ScaleHashHeight;
+	FAudioMaterialMeterStyle& SetScaleHashHeight(float InScaleHashHeight) { ScaleHashHeight = InScaleHashHeight; return *this; }
+
+	// How wide to draw the decibel scale, if it's enabled
+	UPROPERTY(EditAnywhere, BlueprintReadWrite, Category = Appearance, meta = (UIMin = "3", ClampMin = "3", UIMax = "10"))
+	int32 DecibelsPerHash;
+	FAudioMaterialMeterStyle& SetDecibelsPerHash(float InDecibelsPerHash) { DecibelsPerHash = InDecibelsPerHash; return *this; }
+
+	/** Font family and size to be used when displaying the meter scale. */
+	UPROPERTY(EditAnywhere, BlueprintReadWrite, Category = Appearance)
+	FSlateFontInfo Font;
+	FAudioMaterialMeterStyle& SetFont(const FSlateFontInfo& InFont) { Font = InFont; return *this; }
+	FAudioMaterialMeterStyle& SetFont(TSharedPtr<const FCompositeFont> InCompositeFont, const int32 InSize, const FName& InTypefaceFontName = NAME_None) { Font = FSlateFontInfo(InCompositeFont, InSize, InTypefaceFontName); return *this; }
+	FAudioMaterialMeterStyle& SetFont(const UObject* InFontObject, const int32 InSize, const FName& InTypefaceFontName = NAME_None) { Font = FSlateFontInfo(InFontObject, InSize, InTypefaceFontName); return *this; }
+	FAudioMaterialMeterStyle& SetFont(const FName& InFontName, uint16 InSize) { Font = FSlateFontInfo(InFontName, InSize); return *this; }
+	FAudioMaterialMeterStyle& SetFont(const FString& InFontName, uint16 InSize) { Font = FSlateFontInfo(*InFontName, InSize); return *this; }
+	FAudioMaterialMeterStyle& SetFont(const WIDECHAR* InFontName, uint16 InSize) { Font = FSlateFontInfo(InFontName, InSize); return *this; }
+	FAudioMaterialMeterStyle& SetFont(const ANSICHAR* InFontName, uint16 InSize) { Font = FSlateFontInfo(InFontName, InSize); return *this; }
+	FAudioMaterialMeterStyle& SetFontName(const FName& InFontName) { Font = FSlateFontInfo(InFontName, Font.Size); return *this; }
+	FAudioMaterialMeterStyle& SetFontName(const FString& InFontName) { Font = FSlateFontInfo(InFontName, Font.Size); return *this; }
+	FAudioMaterialMeterStyle& SetFontName(const WIDECHAR* InFontName) { Font = FSlateFontInfo(InFontName, Font.Size); return *this; }
+	FAudioMaterialMeterStyle& SetFontName(const ANSICHAR* InFontName) { Font = FSlateFontInfo(InFontName, Font.Size); return *this; }
+	FAudioMaterialMeterStyle& SetFontSize(uint16 InSize) { Font.Size = InSize; return *this; }
+	FAudioMaterialMeterStyle& SetTypefaceFontName(const FName& InTypefaceFontName) { Font.TypefaceFontName = InTypefaceFontName; return *this; }
+
 };
 
 /**
diff --git a/Engine/Plugins/Runtime/AudioWidgets/Source/AudioWidgets/Public/AudioMaterialSlate/SAudioMaterialMeter.h b/Engine/Plugins/Runtime/AudioWidgets/Source/AudioWidgets/Public/AudioMaterialSlate/SAudioMaterialMeter.h
--- a/Engine/Plugins/Runtime/AudioWidgets/Source/AudioWidgets/Public/AudioMaterialSlate/SAudioMaterialMeter.h
+++ b/Engine/Plugins/Runtime/AudioWidgets/Source/AudioWidgets/Public/AudioMaterialSlate/SAudioMaterialMeter.h
@@ -9,6 +9,7 @@
 #include "Widgets/SLeafWidget.h"
 
 class UWidget;
+struct FMeterChannelInfo;
 
 /**
  * A simple slate that renders the meter in single material and modifies the material on value change.
@@ -18,19 +19,20 @@ class AUDIOWIDGETS_API SAudioMaterialMeter : public SLeafWidget
 {
 public:
 	SLATE_BEGIN_ARGS(SAudioMaterialMeter)
+	: _Orientation(EOrientation::Orient_Vertical)
 	{}
 
+	/** The meter's orientation. */
+	SLATE_ARGUMENT(EOrientation, Orientation)
+
 	/** The owner object*/
 	SLATE_ARGUMENT(TWeakObjectPtr<UObject>, Owner)
 
-	/** A value that drives how the Meter is rendered*/
-	SLATE_ATTRIBUTE(float, ValueAttribute)
-
 	/** The style used to draw the meter. */
 	SLATE_STYLE_ARGUMENT(FAudioMaterialMeterStyle, AudioMaterialMeterStyle)
 
-	/** Called when the value is changed in the Meter */
-	SLATE_EVENT(FOnFloatValueChanged, OnValueChanged)
+	/** Attribute representing the meter values */
+	SLATE_ATTRIBUTE(TArray<FMeterChannelInfo>, MeterChannelInfo)
 
 	SLATE_END_ARGS()
 
@@ -39,29 +41,27 @@ public:
 	*/
 	void Construct(const FArguments& InArgs);
 
-	// SWidget overrides
+	//SWidget
 	virtual int32 OnPaint(const FPaintArgs& Args, const FGeometry& AllottedGeometry, const FSlateRect& MyCullingRect, FSlateWindowElementList& OutDrawElements, int32 LayerId, const FWidgetStyle& InWidgetStyle, bool bParentEnabled) const override;
 	virtual FVector2D ComputeDesiredSize(float) const override;
+	//~SWidget
 
-	/** Set the Value attribute */
-	void SetValue(float InValueAttribute);
+	/**Set the Orientation attribute*/
+	void SetOrientation(EOrientation InOrientation);
 
-	/** Apply new material to be used to render the Slate.*/
-	void ApplyNewMaterial();
+	/** Apply new material to be used to render the Meter.*/
+	TArray<TWeakObjectPtr<UMaterialInstanceDynamic>> ApplyNewMaterial();
 
-public:
+	/**Set the MeterChannelInfo attribute*/
+	void SetMeterChannelInfo(const TAttribute<TArray<FMeterChannelInfo>>& InMeterChannelInfo);
 
-	// Holds a delegate that is executed when the slider's value changed.
-	FOnFloatValueChanged OnValueChanged;	
+	/**Get the MeterChannelInfo attribute*/
+	TArray<FMeterChannelInfo> GetMeterChannelInfo() const;
 
 private:
 
-	/**
-	* Commits the specified meter value.
-	*
-	* @param NewValue The value to commit.
-	*/
-	void CommitValue(float NewValue);
+	// Returns the scale width based off font size and hash width
+	float GetScaleWidth() const;
 
 private:
 
@@ -69,9 +69,15 @@ private:
 	TWeakObjectPtr<UObject> Owner;
 
 	// Holds the style for the Slate
-	const FAudioMaterialMeterStyle* AudioMaterialMeterStyle = nullptr;
+	const FAudioMaterialMeterStyle* Style = nullptr;
+
+	// Holds the Modifiable Materials that represent the Meters
+	mutable TArray<TWeakObjectPtr<UMaterialInstanceDynamic>> DynamicMaterials;
+
+	// Holds the Meter's orientation.
+	EOrientation Orientation;
 
-	//Holds the current value
-	TAttribute<float> ValueAttribute = 0.f;
+	//Holds the MeterChannelInfoAttributes
+	TAttribute<TArray<FMeterChannelInfo>> MeterChannelInfoAttribute;
 
 };
