AudioMaterialMeter option to MetasoundEditor
#rnx

diff --git a/Engine/Plugins/Runtime/AudioWidgets/Source/AudioWidgets/Private/AudioMeter.cpp b/Engine/Plugins/Runtime/AudioWidgets/Source/AudioWidgets/Private/AudioMeter.cpp
--- a/Engine/Plugins/Runtime/AudioWidgets/Source/AudioWidgets/Private/AudioMeter.cpp
+++ b/Engine/Plugins/Runtime/AudioWidgets/Source/AudioWidgets/Private/AudioMeter.cpp
@@ -3,13 +3,13 @@
 #include "AudioMeter.h"
 
 #include "AudioBusSubsystem.h"
+#include "AudioMaterialSlate/SAudioMaterialMeter.h"
 #include "AudioMixerDevice.h"
 #include "SAudioMeter.h"
 #include "UObject/UObjectGlobals.h"
 
 #include UE_INLINE_GENERATED_CPP_BY_NAME(AudioMeter)
 
-
 #define LOCTEXT_NAMESPACE "AUDIO_UMG"
 UAudioMeter::UAudioMeter(const FObjectInitializer& ObjectInitializer)
 	: Super(ObjectInitializer)
@@ -167,20 +167,27 @@ namespace AudioWidgets
 		
 	}
 
-	FAudioMeter::FAudioMeter(const int32 InNumChannels, const Audio::FDeviceId InAudioDeviceId, const TObjectPtr<UAudioBus> InExternalAudioBus)
-		: Widget(SNew(SAudioMeter)
-			.Orientation(EOrientation::Orient_Vertical)
-			.BackgroundColor(FLinearColor::Transparent)
-
-			// TODO: Move to editor style
-			.MeterBackgroundColor(FLinearColor(0.031f, 0.031f, 0.031f, 1.0f))
-			.MeterValueColor(FLinearColor(0.025719f, 0.208333f, 0.069907f, 1.0f))
-			.MeterPeakColor(FLinearColor(0.24349f, 0.708333f, 0.357002f, 1.0f))
-			.MeterClippingColor(FLinearColor(1.0f, 0.0f, 0.112334f, 1.0f))
-			.MeterScaleColor(FLinearColor(0.017642f, 0.017642f, 0.017642f, 1.0f))
-			.MeterScaleLabelColor(FLinearColor(0.442708f, 0.442708f, 0.442708f, 1.0f))
-		)
+	FAudioMeter::FAudioMeter(const int32 InNumChannels, const Audio::FDeviceId InAudioDeviceId, const TObjectPtr<UAudioBus> InExternalAudioBus, bool InbUseAudioMaterialWidget)
 	{
+		if (InbUseAudioMaterialWidget)
+		{
+			Widget = (SNew(SAudioMaterialMeter));
+		}
+		else
+		{
+			Widget = SNew(SAudioMeter)
+				.Orientation(EOrientation::Orient_Vertical)
+				.BackgroundColor(FLinearColor::Transparent)
+
+				// TODO: Move to editor style
+				.MeterBackgroundColor(FLinearColor(0.031f, 0.031f, 0.031f, 1.0f))
+				.MeterValueColor(FLinearColor(0.025719f, 0.208333f, 0.069907f, 1.0f))
+				.MeterPeakColor(FLinearColor(0.24349f, 0.708333f, 0.357002f, 1.0f))
+				.MeterClippingColor(FLinearColor(1.0f, 0.0f, 0.112334f, 1.0f))
+				.MeterScaleColor(FLinearColor(0.017642f, 0.017642f, 0.017642f, 1.0f))
+				.MeterScaleLabelColor(FLinearColor(0.442708f, 0.442708f, 0.442708f, 1.0f));
+		}
+
 		Init(InNumChannels, InAudioDeviceId, InExternalAudioBus);
 	}
 
diff --git a/Engine/Plugins/Runtime/AudioWidgets/Source/AudioWidgets/Public/AudioMaterialSlate/SAudioMaterialMeter.h b/Engine/Plugins/Runtime/AudioWidgets/Source/AudioWidgets/Public/AudioMaterialSlate/SAudioMaterialMeter.h
--- a/Engine/Plugins/Runtime/AudioWidgets/Source/AudioWidgets/Public/AudioMaterialSlate/SAudioMaterialMeter.h
+++ b/Engine/Plugins/Runtime/AudioWidgets/Source/AudioWidgets/Public/AudioMaterialSlate/SAudioMaterialMeter.h
@@ -5,9 +5,9 @@
 #include "AudioMaterialSlate/AudioMaterialSlateTypes.h"
 #include "AudioMeterTypes.h"
 #include "Framework/SlateDelegates.h"
+#include "SAudioMeter.h"
 #include "Styling/ISlateStyle.h"
 #include "Styling/SlateWidgetStyleAsset.h"
-#include "Widgets/SLeafWidget.h"
 
 class UWidget;
 struct FMeterChannelInfo;
@@ -16,7 +16,7 @@ struct FMeterChannelInfo;
  * A simple slate that renders the meter in single material and modifies the material on value change.
  *
  */
-class AUDIOWIDGETS_API SAudioMaterialMeter : public SLeafWidget
+class AUDIOWIDGETS_API SAudioMaterialMeter : public SAudioMeterBase
 {
 public:
 	SLATE_BEGIN_ARGS(SAudioMaterialMeter)
@@ -54,10 +54,10 @@ public:
 	TArray<TWeakObjectPtr<UMaterialInstanceDynamic>> ApplyNewMaterial();
 
 	/**Set the MeterChannelInfo attribute*/
-	void SetMeterChannelInfo(const TAttribute<TArray<FMeterChannelInfo>>& InMeterChannelInfo);
+	void SetMeterChannelInfo(const TAttribute<TArray<FMeterChannelInfo>>& InMeterChannelInfo) override;
 
 	/**Get the MeterChannelInfo attribute*/
-	TArray<FMeterChannelInfo> GetMeterChannelInfo() const;
+	TArray<FMeterChannelInfo> GetMeterChannelInfo() const override;
 
 private:
 
diff --git a/Engine/Plugins/Runtime/AudioWidgets/Source/AudioWidgets/Public/AudioMeter.h b/Engine/Plugins/Runtime/AudioWidgets/Source/AudioWidgets/Public/AudioMeter.h
--- a/Engine/Plugins/Runtime/AudioWidgets/Source/AudioWidgets/Public/AudioMeter.h
+++ b/Engine/Plugins/Runtime/AudioWidgets/Source/AudioWidgets/Public/AudioMeter.h
@@ -21,6 +21,7 @@
 #include "AudioMeter.generated.h"
 
 // Forward Declarations
+class SAudioMaterialMeter;
 class SAudioMeter;
 class UWorld;
 
@@ -157,15 +158,22 @@ namespace AudioWidgets
 	{
 	public:
 		UE_DEPRECATED(5.4, "Use the FAudioMeter constructor that uses Audio::FDeviceId.")
-		FAudioMeter(int32 InNumChannels, UWorld& InWorld, TObjectPtr<UAudioBus> InExternalAudioBus = nullptr); // OPTIONAL PARAM InExternalAudioBus: An audio meter can be constructed from this audio bus.
+		FAudioMeter(int32 InNumChannels, UWorld& InWorld, TObjectPtr<UAudioBus> InExternalAudioBus = nullptr); 
 		
-		FAudioMeter(const int32 InNumChannels, const Audio::FDeviceId InAudioDeviceId, TObjectPtr<UAudioBus> InExternalAudioBus = nullptr); // OPTIONAL PARAM InExternalAudioBus: An audio meter can be constructed from this audio bus.
+		//** OPTIONAL PARAM InExternalAudioBus: An audio meter can be constructed from this audio bus. , InbUseAudioMaterialWidget: Is the AudioMaterialWidgets used for analyzer visualization.*/
+		FAudioMeter(const int32 InNumChannels, const Audio::FDeviceId InAudioDeviceId, TObjectPtr<UAudioBus> InExternalAudioBus = nullptr, bool InbUseAudioMaterialWidget = false); 
 		~FAudioMeter();
 
 		UAudioBus* GetAudioBus() const;
 
 		TSharedRef<SAudioMeter> GetWidget() const;
 
+		template<class T>
+		TSharedRef<T> GetWidget() const
+		{
+			return StaticCastSharedRef<T>(Widget->AsShared());
+		};
+
 		UE_DEPRECATED(5.4, "Use the Init method that uses Audio::FDeviceId.")
 		void Init(int32 InNumChannels, UWorld& InWorld, TObjectPtr<UAudioBus> InExternalAudioBus = nullptr);
 
@@ -193,7 +201,7 @@ namespace AudioWidgets
 		TStrongObjectPtr<UMeterSettings> Settings;
 
 		/** MetaSound Output Meter widget */
-		TSharedPtr<SAudioMeter> Widget;
+		TSharedPtr<SAudioMeterBase> Widget;
 
 		bool bUseExternalAudioBus = false;
 	};
diff --git a/Engine/Plugins/Runtime/AudioWidgets/Source/AudioWidgets/Public/SAudioMeter.h b/Engine/Plugins/Runtime/AudioWidgets/Source/AudioWidgets/Public/SAudioMeter.h
--- a/Engine/Plugins/Runtime/AudioWidgets/Source/AudioWidgets/Public/SAudioMeter.h
+++ b/Engine/Plugins/Runtime/AudioWidgets/Source/AudioWidgets/Public/SAudioMeter.h
@@ -12,11 +12,30 @@
 class FPaintArgs;
 class FSlateWindowElementList;
 
+class AUDIOWIDGETS_API SAudioMeterBase
+	: public SLeafWidget
+{
+public:
+
+	SLATE_BEGIN_ARGS(SAudioMeterBase){}
+
+	SLATE_END_ARGS()
+
+	/** Is the active timer registered to refresh the meter channel info. */
+	bool bIsActiveTimerRegistered = false;
+
+public:
+
+	virtual void SetMeterChannelInfo(const TAttribute<TArray<FMeterChannelInfo>>& InMeterChannelInfo) = 0;
+	virtual TArray<FMeterChannelInfo> GetMeterChannelInfo() const = 0;
+
+};
+
 /**
  * A Slate slider control is a linear scale and draggable handle.
  */
 class AUDIOWIDGETS_API SAudioMeter
-	: public SLeafWidget
+	: public SAudioMeterBase
 {
 public:
 
@@ -76,8 +95,8 @@ public:
 	 */
 	void Construct(const SAudioMeter::FArguments& InDeclaration);
 
-	void SetMeterChannelInfo(const TAttribute<TArray<FMeterChannelInfo>>& InMeterChannelInfo);
-	TArray<FMeterChannelInfo> GetMeterChannelInfo() const;
+	void SetMeterChannelInfo(const TAttribute<TArray<FMeterChannelInfo>>& InMeterChannelInfo) override;
+	TArray<FMeterChannelInfo> GetMeterChannelInfo() const override;
 
 	/** Set the Orientation attribute */
 	void SetOrientation(EOrientation InOrientation);
@@ -90,9 +109,6 @@ public:
 	void SetMeterScaleColor(FSlateColor InMeterScaleColor);
 	void SetMeterScaleLabelColor(FSlateColor InMeterScaleLabelColor);
 
-	/** Is the active timer registered to refresh the meter channel info. */
-	bool bIsActiveTimerRegistered = false;
-
 public:
 
 	// SWidget overrides
diff --git a/Engine/Plugins/Runtime/Metasound/Source/MetasoundEditor/Private/MetasoundEditor.cpp b/Engine/Plugins/Runtime/Metasound/Source/MetasoundEditor/Private/MetasoundEditor.cpp
--- a/Engine/Plugins/Runtime/Metasound/Source/MetasoundEditor/Private/MetasoundEditor.cpp
+++ b/Engine/Plugins/Runtime/Metasound/Source/MetasoundEditor/Private/MetasoundEditor.cpp
@@ -4,6 +4,7 @@
 
 #include "Algo/AnyOf.h"
 #include "AudioDevice.h"
+#include "AudioMaterialSlate/SAudioMaterialMeter.h"
 #include "AudioMeterStyle.h"
 #include "AudioOscilloscope.h"
 #include "AudioSpectrumAnalyzer.h"
@@ -92,6 +93,17 @@ namespace Metasound
 {
 	namespace Editor
 	{
+		namespace EditorPrivate
+		{
+			int32 UseAudioMaterialWidgets = 0;
+			FAutoConsoleVariableRef CVarUseAudioMaterialWidgets(
+				TEXT("au.MetaSound.Editor.EnableAudioMaterialAnalyzers"),
+				UseAudioMaterialWidgets,
+				TEXT("Are new AudioMaterialWidgets used for visualization for the Metasound Analyzers, if implemented.\n")
+				TEXT("0: Disabled (default), !0: Enabled"),
+				ECVF_Default);
+		}//Metasound::Editor::EditorPrivate
+
 		static const TArray<FText> NodeSectionNames
 		{
 			LOCTEXT("NodeSectionName_Invalid", "INVALID"),
@@ -1233,7 +1245,9 @@ namespace Metasound
 
 					if (!OutputMeter.IsValid())
 					{
-						OutputMeter = MakeShared<AudioWidgets::FAudioMeter>(MetaSoundSource->NumChannels, AudioDeviceId);
+
+						bool bUseAudioMaterialWidgets = EditorPrivate::UseAudioMaterialWidgets == 0 ? false : true;
+						OutputMeter = MakeShared<AudioWidgets::FAudioMeter>(MetaSoundSource->NumChannels, AudioDeviceId, nullptr, bUseAudioMaterialWidgets);
 					}
 					else if (OutputMeter->GetAudioBus()->GetNumChannels() != MetaSoundSource->NumChannels)
 					{
@@ -1796,7 +1810,7 @@ namespace Metasound
 					})
 				);
 
-				TSharedPtr<SAudioMeter> OutputMeterWidget = OutputMeter->GetWidget();
+				TSharedPtr<SAudioMeterBase> OutputMeterWidget = OutputMeter->GetWidget<SAudioMeterBase>();
 				if (OutputMeterWidget.IsValid())
 				{
 					if (!OutputMeterWidget->bIsActiveTimerRegistered)
@@ -1810,7 +1824,7 @@ namespace Metasound
 								}
 								else
 								{
-									TSharedRef<SAudioMeter> MeterRef = OutputMeter->GetWidget();
+									TSharedRef<SAudioMeterBase> MeterRef = OutputMeter->GetWidget<SAudioMeterBase>();
 									MeterRef->bIsActiveTimerRegistered = false;
 									return EActiveTimerReturnType::Stop;
 								}
diff --git a/Engine/Plugins/Runtime/Metasound/Source/MetasoundEditor/Private/SMetasoundGraphNode.cpp b/Engine/Plugins/Runtime/Metasound/Source/MetasoundEditor/Private/SMetasoundGraphNode.cpp
--- a/Engine/Plugins/Runtime/Metasound/Source/MetasoundEditor/Private/SMetasoundGraphNode.cpp
+++ b/Engine/Plugins/Runtime/Metasound/Source/MetasoundEditor/Private/SMetasoundGraphNode.cpp
@@ -66,9 +66,9 @@ namespace Metasound
 		{
 			int32 UseAudioMaterialWidgets = 0;
 			FAutoConsoleVariableRef CVarUseAudioMaterialWidgets(
-				TEXT("au.MetaSound.Editor.UseAudioMaterialSlates"),
+				TEXT("au.MetaSound.Editor.EnableAudioMaterialNodes"),
 				UseAudioMaterialWidgets,
-				TEXT("Are new AudioMaterialWidgets used for visualization in the Metasound editor, if implemented.\n")
+				TEXT("Are new AudioMaterialWidgets used for visualization for the Metasound nodes, if implemented.\n")
 				TEXT("0: Disabled (default), !0: Enabled"),
 				ECVF_Default);
 		}// Metasound::Editor::GraphNodePrivate
