[AudioWidgets] Adjusted audio widgets UFUNCTION categories so they are easier to find when searching for actions in Blueprints

diff --git a/Engine/Plugins/Runtime/AudioWidgets/Source/AudioWidgets/Public/AudioMaterialSlate/AudioMaterialButton.h b/Engine/Plugins/Runtime/AudioWidgets/Source/AudioWidgets/Public/AudioMaterialSlate/AudioMaterialButton.h
--- a/Engine/Plugins/Runtime/AudioWidgets/Source/AudioWidgets/Public/AudioMaterialSlate/AudioMaterialButton.h
+++ b/Engine/Plugins/Runtime/AudioWidgets/Source/AudioWidgets/Public/AudioMaterialSlate/AudioMaterialButton.h
@@ -47,11 +47,11 @@ public:
 	// End of UVisual
 
 		/** Gets the current value of the slider.*/
-	UFUNCTION(BlueprintPure, Category = "Behavior")
+	UFUNCTION(BlueprintPure, Category = "Audio Widgets| Audio Material Button")
 	bool GetIsPressed() const;
 
 	/** Sets the current value of the slider. InValue is Clamped between 0.f - 1.f */
-	UFUNCTION(BlueprintCallable, Category = "Behavior")
+	UFUNCTION(BlueprintCallable, Category = "Audio Widgets| Audio Material Button")
 	void SetIsPressed(bool InPressed);
 
 public:
diff --git a/Engine/Plugins/Runtime/AudioWidgets/Source/AudioWidgets/Public/AudioMaterialSlate/AudioMaterialKnob.h b/Engine/Plugins/Runtime/AudioWidgets/Source/AudioWidgets/Public/AudioMaterialSlate/AudioMaterialKnob.h
--- a/Engine/Plugins/Runtime/AudioWidgets/Source/AudioWidgets/Public/AudioMaterialSlate/AudioMaterialKnob.h
+++ b/Engine/Plugins/Runtime/AudioWidgets/Source/AudioWidgets/Public/AudioMaterialSlate/AudioMaterialKnob.h
@@ -45,51 +45,51 @@ public:
 	// End of UVisual
 
 	/** Get the current value of the knob.*/
-	UFUNCTION(BlueprintPure, Category = "Behavior")
+	UFUNCTION(BlueprintPure, Category = "Audio Widgets| Audio Material Knob")
 	float GetValue();
 
 	/** Set the current value of the knob. InValue is Clamped between 0.f - 1.f */
-	UFUNCTION(BlueprintCallable, Category = "Behavior")
+	UFUNCTION(BlueprintCallable, Category = "Audio Widgets| Audio Material Knob")
 	void SetValue(float InValue);	
 	
 	/** Set the knobs tune speed. InValue is Clamped between 0.f - 1.f */
-	UFUNCTION(BlueprintCallable, Category = "Behavior")
+	UFUNCTION(BlueprintCallable, Category = "Audio Widgets| Audio Material Knob")
 	void SetTuneSpeed(float InValue);	
 
 	/** Get the Knobs tune speed*/
-	UFUNCTION(BlueprintPure, Category = "Behavior")
+	UFUNCTION(BlueprintPure, Category = "Audio Widgets| Audio Material Knob")
 	float GetTuneSpeed() const;
 	
 	/** Set the knobs fine-tune speed. InValue is Clamped between 0.f - 1.f */
-	UFUNCTION(BlueprintCallable, Category = "Behavior")
+	UFUNCTION(BlueprintCallable, Category = "Audio Widgets| Audio Material Knob")
 	void SetFineTuneSpeed(float InValue);
 
 	/** Get the Knobs fine-tune speed*/
-	UFUNCTION(BlueprintPure, Category = "Behavior")
+	UFUNCTION(BlueprintPure, Category = "Audio Widgets| Audio Material Knob")
 	float GetFineTuneSpeed() const;
 	
 	/** Set the knob to be interactive or fixed */
-	UFUNCTION(BlueprintCallable, Category = "Behavior")
+	UFUNCTION(BlueprintCallable, Category = "Audio Widgets| Audio Material Knob")
 	void SetLocked(bool InLocked);	
 
 	/** Get whether the knob is interactive or fixed.*/
-	UFUNCTION(BlueprintPure, Category = "Behavior")
+	UFUNCTION(BlueprintPure, Category = "Audio Widgets| Audio Material Knob")
 	bool GetIsLocked() const;
 	
 	/** Set the knob to use steps when turning On Mouse move */
-	UFUNCTION(BlueprintCallable, Category = "Behavior")
+	UFUNCTION(BlueprintCallable, Category = "Audio Widgets| Audio Material Knob")
 	void SetMouseUsesStep(bool InUsesStep);
 
 	/** Get whether the knob uses steps when tuning On Mouse move*/
-	UFUNCTION(BlueprintPure, Category = "Behavior")
+	UFUNCTION(BlueprintPure, Category = "Audio Widgets| Audio Material Knob")
 	bool GetMouseUsesStep() const;
 
 	/** Set the amount to adjust the value when using steps*/
-	UFUNCTION(BlueprintCallable, Category = "Behavior")
+	UFUNCTION(BlueprintCallable, Category = "Audio Widgets| Audio Material Knob")
 	void SetStepSize(float InValue);
 
 	/** Get Step Size*/
-	UFUNCTION(BlueprintPure, Category = "Behavior")
+	UFUNCTION(BlueprintPure, Category = "Audio Widgets| Audio Material Knob")
 	float GetStepSize() const;
 
 public:
diff --git a/Engine/Plugins/Runtime/AudioWidgets/Source/AudioWidgets/Public/AudioMaterialSlate/AudioMaterialMeter.h b/Engine/Plugins/Runtime/AudioWidgets/Source/AudioWidgets/Public/AudioMaterialSlate/AudioMaterialMeter.h
--- a/Engine/Plugins/Runtime/AudioWidgets/Source/AudioWidgets/Public/AudioMaterialSlate/AudioMaterialMeter.h
+++ b/Engine/Plugins/Runtime/AudioWidgets/Source/AudioWidgets/Public/AudioMaterialSlate/AudioMaterialMeter.h
@@ -49,11 +49,11 @@ public:
 	// End of UVisual
 
 	/** Gets the current linear values of the meter. */
-	UFUNCTION(BlueprintCallable, Category = "Behavior")
+	UFUNCTION(BlueprintCallable, Category = "Audio Widgets| Audio Material Meter")
 	TArray<FMeterChannelInfo> GetMeterChannelInfo() const;
 
 	/** Sets the current meter values. */
-	UFUNCTION(BlueprintCallable, Category = "Behavior")
+	UFUNCTION(BlueprintCallable, Category = "Audio Widgets| Audio Material Meter")
 	void SetMeterChannelInfo(const TArray<FMeterChannelInfo>& InMeterChannelInfo);
 
 public:
diff --git a/Engine/Plugins/Runtime/AudioWidgets/Source/AudioWidgets/Public/AudioMaterialSlate/AudioMaterialSlider.h b/Engine/Plugins/Runtime/AudioWidgets/Source/AudioWidgets/Public/AudioMaterialSlate/AudioMaterialSlider.h
--- a/Engine/Plugins/Runtime/AudioWidgets/Source/AudioWidgets/Public/AudioMaterialSlate/AudioMaterialSlider.h
+++ b/Engine/Plugins/Runtime/AudioWidgets/Source/AudioWidgets/Public/AudioMaterialSlate/AudioMaterialSlider.h
@@ -45,51 +45,51 @@ public:
 	// End of UVisual
 
 	/** Gets the current value of the slider.*/
-	UFUNCTION(BlueprintPure, Category = "Behavior")
+	UFUNCTION(BlueprintPure, Category = "Audio Widgets| Audio Material Slider")
 	float GetValue() const;
 
 	/** Sets the current value of the slider. InValue is Clamped between 0.f - 1.f */
-	UFUNCTION(BlueprintCallable, Category = "Behavior")
+	UFUNCTION(BlueprintCallable, Category = "Audio Widgets| Audio Material Slider")
 	void SetValue(float InValue);
 
 	/** Set the tune speed of the slider. InValue is Clamped between 0.f - 1.f */
-	UFUNCTION(BlueprintCallable, Category = "Behavior")
+	UFUNCTION(BlueprintCallable, Category = "Audio Widgets| Audio Material Slider")
 	void SetTuneSpeed(const float InValue);
 
 	/** Get slider tune speed*/
-	UFUNCTION(BlueprintPure, Category = "Behavior")
+	UFUNCTION(BlueprintPure, Category = "Audio Widgets| Audio Material Slider")
 	float GetTuneSpeed() const;
 
 	/** Set the fine-tune speed of the slider. InValue is Clamped between 0.f - 1.f */
-	UFUNCTION(BlueprintCallable, Category = "Behavior")
+	UFUNCTION(BlueprintCallable, Category = "Audio Widgets| Audio Material Slider")
 	void SetFineTuneSpeed(const float InValue);
 
 	/** Get slider fine-tune speed*/
-	UFUNCTION(BlueprintPure, Category = "Behavior")
+	UFUNCTION(BlueprintPure, Category = "Audio Widgets| Audio Material Slider")
 	float GetFineTuneSpeed() const;
 
 	/** Set the slider to be interactive or fixed */
-	UFUNCTION(BlueprintCallable, Category = "Behavior")
+	UFUNCTION(BlueprintCallable, Category = "Audio Widgets| Audio Material Slider")
 	void SetLocked(bool bInLocked);
 
 	/** Get whether the slider is interactive or fixed.*/
-	UFUNCTION(BlueprintPure, Category = "Behavior")
+	UFUNCTION(BlueprintPure, Category = "Audio Widgets| Audio Material Slider")
 	bool GetIsLocked() const;
 
 	/** Sets the slider to use steps when turning On Mouse move */
-	UFUNCTION(BlueprintCallable, Category = "Behavior")
+	UFUNCTION(BlueprintCallable, Category = "Audio Widgets| Audio Material Slider")
 	void SetMouseUsesStep(bool bInUsesStep);
 
 	/** Get whether the slider uses steps when turning On Mouse move*/
-	UFUNCTION(BlueprintPure, Category = "Behavior")
+	UFUNCTION(BlueprintPure, Category = "Audio Widgets| Audio Material Slider")
 	bool GetMouseUsesStep() const;
 
 	/** Sets the amount to adjust the value when using steps*/
-	UFUNCTION(BlueprintCallable, Category = "Behavior")
+	UFUNCTION(BlueprintCallable, Category = "Audio Widgets| Audio Material Slider")
 	void SetStepSize(float InValue);
 
 	/** Get Step Size*/
-	UFUNCTION(BlueprintPure, Category = "Behavior")
+	UFUNCTION(BlueprintPure, Category = "Audio Widgets| Audio Material Slider")
 	float GetStepSize() const;
 
 public:
diff --git a/Engine/Plugins/Runtime/AudioWidgets/Source/AudioWidgets/Public/AudioMeter.h b/Engine/Plugins/Runtime/AudioWidgets/Source/AudioWidgets/Public/AudioMeter.h
--- a/Engine/Plugins/Runtime/AudioWidgets/Source/AudioWidgets/Public/AudioMeter.h
+++ b/Engine/Plugins/Runtime/AudioWidgets/Source/AudioWidgets/Public/AudioMeter.h
@@ -94,39 +94,39 @@ public:
 public:
 
  	/** Gets the current linear value of the meter. */
- 	UFUNCTION(BlueprintCallable, Category="Behavior")
+ 	UFUNCTION(BlueprintCallable, Category = "Audio Widgets| Audio Meter")
 	TArray<FMeterChannelInfo> GetMeterChannelInfo() const;
  
  	/** Sets the current meter values. */
- 	UFUNCTION(BlueprintCallable, Category="Behavior")
+ 	UFUNCTION(BlueprintCallable, Category = "Audio Widgets| Audio Meter")
  	void SetMeterChannelInfo(const TArray<FMeterChannelInfo>& InMeterChannelInfo);
 
 	/** Sets the background color */
-	UFUNCTION(BlueprintCallable, Category = "Appearance")
+	UFUNCTION(BlueprintCallable, Category = "Audio Widgets| Audio Meter")
 	void SetBackgroundColor(FLinearColor InValue);
 
 	/** Sets the meter background color */
-	UFUNCTION(BlueprintCallable, Category="Appearance")
+	UFUNCTION(BlueprintCallable, Category = "Audio Widgets| Audio Meter")
 	void SetMeterBackgroundColor(FLinearColor InValue);
 
 	/** Sets the meter value color */
-	UFUNCTION(BlueprintCallable, Category = "Appearance")
+	UFUNCTION(BlueprintCallable, Category = "Audio Widgets| Audio Meter")
 	void SetMeterValueColor(FLinearColor InValue);
 
 	/** Sets the meter peak color */
-	UFUNCTION(BlueprintCallable, Category = "Appearance")
+	UFUNCTION(BlueprintCallable, Category = "Audio Widgets| Audio Meter")
 	void SetMeterPeakColor(FLinearColor InValue);
 
 	/** Sets the meter clipping color */
-	UFUNCTION(BlueprintCallable, Category = "Appearance")
+	UFUNCTION(BlueprintCallable, Category = "Audio Widgets| Audio Meter")
 	void SetMeterClippingColor(FLinearColor InValue);
 
 	/** Sets the meter scale color */
-	UFUNCTION(BlueprintCallable, Category = "Appearance")
+	UFUNCTION(BlueprintCallable, Category = "Audio Widgets| Audio Meter")
 	void SetMeterScaleColor(FLinearColor InValue);
 
 	/** Sets the meter scale color */
-	UFUNCTION(BlueprintCallable, Category = "Appearance")
+	UFUNCTION(BlueprintCallable, Category = "Audio Widgets| Audio Meter")
 	void SetMeterScaleLabelColor(FLinearColor InValue);
 
 	// UWidget interface
diff --git a/Engine/Plugins/Runtime/AudioWidgets/Source/AudioWidgets/Public/AudioOscilloscopeUMG.h b/Engine/Plugins/Runtime/AudioWidgets/Source/AudioWidgets/Public/AudioOscilloscopeUMG.h
--- a/Engine/Plugins/Runtime/AudioWidgets/Source/AudioWidgets/Public/AudioOscilloscopeUMG.h
+++ b/Engine/Plugins/Runtime/AudioWidgets/Source/AudioWidgets/Public/AudioOscilloscopeUMG.h
@@ -39,11 +39,11 @@ public:
 #endif
 
 	/** Starts the oscilloscope processing. */
-	UFUNCTION(BlueprintCallable, Category = "Behavior")
+	UFUNCTION(BlueprintCallable, Category = "Audio Widgets| Oscilloscope")
 	void StartProcessing();
 
 	/** Stops the oscilloscope processing. */
-	UFUNCTION(BlueprintCallable, Category = "Behavior")
+	UFUNCTION(BlueprintCallable, Category = "Audio Widgets| Oscilloscope")
 	void StopProcessing();
 
 	/** The oscilloscope panel style */
diff --git a/Engine/Plugins/Runtime/AudioWidgets/Source/AudioWidgets/Public/AudioRadialSlider.h b/Engine/Plugins/Runtime/AudioWidgets/Source/AudioWidgets/Public/AudioRadialSlider.h
--- a/Engine/Plugins/Runtime/AudioWidgets/Source/AudioWidgets/Public/AudioRadialSlider.h
+++ b/Engine/Plugins/Runtime/AudioWidgets/Source/AudioWidgets/Public/AudioRadialSlider.h
@@ -86,63 +86,63 @@ public:
 
 public: 
 	/** Get output value from normalized linear (0 - 1) based on internal lin to output mapping. */
-	UFUNCTION(BlueprintCallable, Category = "Behavior")
+	UFUNCTION(BlueprintCallable, Category = "Audio Widgets| Audio Radial Slider")
 	float GetOutputValue(const float InSliderValue);
 
 	/** Get normalized linear (0 - 1) slider value from output based on internal lin to output mapping. */
-	UFUNCTION(BlueprintCallable, Category = "Behavior")
+	UFUNCTION(BlueprintCallable, Category = "Audio Widgets| Audio Radial Slider")
 	float GetSliderValue(const float OutputValue);
 
 	/** Sets the widget layout */
-	UFUNCTION(BlueprintCallable, Category = "Appearance")
+	UFUNCTION(BlueprintCallable, Category = "Audio Widgets| Audio Radial Slider")
 	void SetWidgetLayout(EAudioRadialSliderLayout InLayout);
 
 	/** Sets the label background color */
-	UFUNCTION(BlueprintCallable, Category = "Appearance")
+	UFUNCTION(BlueprintCallable, Category = "Audio Widgets| Audio Radial Slider")
 	void SetCenterBackgroundColor(FLinearColor InValue);
 
 	/** Sets the slider progress color */
-	UFUNCTION(BlueprintCallable, Category = "Appearance")
+	UFUNCTION(BlueprintCallable, Category = "Audio Widgets| Audio Radial Slider")
 	void SetSliderProgressColor(FLinearColor InValue);
 
 	/** Sets the slider bar color */
-	UFUNCTION(BlueprintCallable, Category = "Appearance")
+	UFUNCTION(BlueprintCallable, Category = "Audio Widgets| Audio Radial Slider")
 	void SetSliderBarColor(FLinearColor InValue);
 
 	/** Sets the start and end of the hand as a ratio to the slider radius (so 0.0 to 1.0 is from the slider center to the handle). */
-	UFUNCTION(BlueprintCallable, Category = "Appearance")
+	UFUNCTION(BlueprintCallable, Category = "Audio Widgets| Audio Radial Slider")
 	void SetHandStartEndRatio(const FVector2D InHandStartEndRatio);
 
 	/** Sets the label background color */
-	UFUNCTION(BlueprintCallable, Category = "Appearance")
+	UFUNCTION(BlueprintCallable, Category = "Audio Widgets| Audio Radial Slider")
 	void SetTextLabelBackgroundColor(FSlateColor InColor);
 
 	/** Sets the units text */
-	UFUNCTION(BlueprintCallable, Category = "Appearance")
+	UFUNCTION(BlueprintCallable, Category = "Audio Widgets| Audio Radial Slider")
 	void SetUnitsText(const FText Units);
 	
 	/** Sets whether the units text is read only*/
-	UFUNCTION(BlueprintCallable, Category = "Appearance")
+	UFUNCTION(BlueprintCallable, Category = "Audio Widgets| Audio Radial Slider")
 	void SetUnitsTextReadOnly(const bool bIsReadOnly);
 
 	/** Sets whether the value text is read only*/
-	UFUNCTION(BlueprintCallable, Category = "Appearance")
+	UFUNCTION(BlueprintCallable, Category = "Audio Widgets| Audio Radial Slider")
 	void SetValueTextReadOnly(const bool bIsReadOnly);
 	
 	/** If true, show text label only on hover; if false always show label. */
-	UFUNCTION(BlueprintCallable, Category = "Appearance")
+	UFUNCTION(BlueprintCallable, Category = "Audio Widgets| Audio Radial Slider")
 	void SetShowLabelOnlyOnHover(const bool bShowLabelOnlyOnHover);
 	
 	/** Sets whether to show the units text */
-	UFUNCTION(BlueprintCallable, Category = "Appearance")
+	UFUNCTION(BlueprintCallable, Category = "Audio Widgets| Audio Radial Slider")
 	void SetShowUnitsText(const bool bShowUnitsText);
 
 	/** Sets the slider thickness */
-	UFUNCTION(BlueprintCallable, Category = "Appearance")
+	UFUNCTION(BlueprintCallable, Category = "Audio Widgets| Audio Radial Slider")
 	void SetSliderThickness(const float InThickness);
 	
 	/** Sets the output range */
-	UFUNCTION(BlueprintCallable, Category = "Behavior")
+	UFUNCTION(BlueprintCallable, Category = "Audio Widgets| Audio Radial Slider")
 	void SetOutputRange(const FVector2D InOutputRange);
 
 	// UWidget interface
diff --git a/Engine/Plugins/Runtime/AudioWidgets/Source/AudioWidgets/Public/AudioSlider.h b/Engine/Plugins/Runtime/AudioWidgets/Source/AudioWidgets/Public/AudioSlider.h
--- a/Engine/Plugins/Runtime/AudioWidgets/Source/AudioWidgets/Public/AudioSlider.h
+++ b/Engine/Plugins/Runtime/AudioWidgets/Source/AudioWidgets/Public/AudioSlider.h
@@ -90,65 +90,65 @@ public:
 
 public:
 	/** Get output value from normalized linear (0 - 1) based on internal lin to output mapping. */
-	UFUNCTION(BlueprintCallable, Category = "Behavior")
+	UFUNCTION(BlueprintCallable, Category = "Audio Widgets| Audio Slider")
 	float GetOutputValue(const float InSliderValue);
 
 	/** Get normalized linear (0 - 1) value from output based on internal lin to output mapping. */
-	UFUNCTION(BlueprintCallable, Category = "Behavior", meta=(DeprecatedFunction, DeprecationMessage="5.1 - GetLinValue is deprecated, please use GetSliderValue instead."))
+	UFUNCTION(BlueprintCallable, Category = "Audio Widgets| Audio Slider", meta=(DeprecatedFunction, DeprecationMessage="5.1 - GetLinValue is deprecated, please use GetSliderValue instead."))
 	float GetLinValue(const float OutputValue);
 
 	/** Get normalized linear (0 - 1) slider value from output based on internal lin to output mapping. */
-	UFUNCTION(BlueprintCallable, Category = "Behavior")
+	UFUNCTION(BlueprintCallable, Category = "Audio Widgets| Audio Slider")
 	float GetSliderValue(const float OutputValue);
 
 	/** Sets the label background color */
-	UFUNCTION(BlueprintCallable, Category = "Appearance")
+	UFUNCTION(BlueprintCallable, Category = "Audio Widgets| Audio Slider")
 	void SetTextLabelBackgroundColor(FSlateColor InColor);
 
 	/** Sets the units text */
-	UFUNCTION(BlueprintCallable, Category = "Appearance")
+	UFUNCTION(BlueprintCallable, Category = "Audio Widgets| Audio Slider")
 	void SetUnitsText(const FText Units);
 	
 	/** Sets whether the units text is read only*/
-	UFUNCTION(BlueprintCallable, Category = "Appearance")
+	UFUNCTION(BlueprintCallable, Category = "Audio Widgets| Audio Slider")
 	void SetUnitsTextReadOnly(const bool bIsReadOnly);
 
 	/** Sets whether the value text is read only*/
-	UFUNCTION(BlueprintCallable, Category = "Appearance")
+	UFUNCTION(BlueprintCallable, Category = "Audio Widgets| Audio Slider")
 	void SetValueTextReadOnly(const bool bIsReadOnly);
 	
 	/** If true, show text label only on hover; if false always show label. */
-	UFUNCTION(BlueprintCallable, Category = "Appearance")
+	UFUNCTION(BlueprintCallable, Category = "Audio Widgets| Audio Slider")
 	void SetShowLabelOnlyOnHover(const bool bShowLabelOnlyOnHover);
 	
 	/** Sets whether to show the units text */
-	UFUNCTION(BlueprintCallable, Category = "Appearance")
+	UFUNCTION(BlueprintCallable, Category = "Audio Widgets| Audio Slider")
 	void SetShowUnitsText(const bool bShowUnitsText);
 
-	/** The slider's orientation. */
-	UPROPERTY(EditAnywhere, BlueprintReadOnly, Category = Appearance)
-	TEnumAsByte<EOrientation> Orientation;
-
-	/** Called when the value is changed by slider or typing. */
-	UPROPERTY(BlueprintAssignable, Category = "Widget Event")
-	FOnFloatValueChangedEvent OnValueChanged;
-
 	/** Sets the slider background color */
-	UFUNCTION(BlueprintCallable, Category = "Appearance")
+	UFUNCTION(BlueprintCallable, Category = "Audio Widgets| Audio Slider")
 	void SetSliderBackgroundColor(FLinearColor InValue);
 
 	/** Sets the slider bar color */
-	UFUNCTION(BlueprintCallable, Category = "Appearance")
+	UFUNCTION(BlueprintCallable, Category = "Audio Widgets| Audio Slider")
 	void SetSliderBarColor(FLinearColor InValue);
 
 	/** Sets the slider thumb color */
-	UFUNCTION(BlueprintCallable, Category = "Appearance")
+	UFUNCTION(BlueprintCallable, Category = "Audio Widgets| Audio Slider")
 	void SetSliderThumbColor(FLinearColor InValue);
 
 	/** Sets the widget background color */
-	UFUNCTION(BlueprintCallable, Category = "Appearance")
+	UFUNCTION(BlueprintCallable, Category = "Audio Widgets| Audio Slider")
 	void SetWidgetBackgroundColor(FLinearColor InValue);
 
+	/** The slider's orientation. */
+	UPROPERTY(EditAnywhere, BlueprintReadOnly, Category = Appearance)
+	TEnumAsByte<EOrientation> Orientation;
+
+	/** Called when the value is changed by slider or typing. */
+	UPROPERTY(BlueprintAssignable, Category = "Widget Event")
+	FOnFloatValueChangedEvent OnValueChanged;
+
 	// UWidget interface
 	virtual void SynchronizeProperties() override;
 	// End of UWidget interface
diff --git a/Engine/Plugins/Runtime/AudioWidgets/Source/AudioWidgets/Public/AudioVectorscopeUMG.h b/Engine/Plugins/Runtime/AudioWidgets/Source/AudioWidgets/Public/AudioVectorscopeUMG.h
--- a/Engine/Plugins/Runtime/AudioWidgets/Source/AudioWidgets/Public/AudioVectorscopeUMG.h
+++ b/Engine/Plugins/Runtime/AudioWidgets/Source/AudioWidgets/Public/AudioVectorscopeUMG.h
@@ -36,11 +36,11 @@ public:
 #endif
 
 	/** Starts the vectorscope processing. */
-	UFUNCTION(BlueprintCallable, Category = "Behavior")
+	UFUNCTION(BlueprintCallable, Category = "Audio Widgets| Vectorscope")
 	void StartProcessing();
 
 	/** Stops the vectorscope processing. */
-	UFUNCTION(BlueprintCallable, Category = "Behavior")
+	UFUNCTION(BlueprintCallable, Category = "Audio Widgets| Vectorscope")
 	void StopProcessing();
 
 	/** The vectorscope panel style */
