[AudioWidgets] Vectorscope: Fixed multiple bugs related to UMG widget:
- Renamed var from TimeWindow to DisplayPersistence since it makes more sense in the context of the Vectorscope
- UAudioVectorscope: Fixed bug in DisplayPersistenceMs UPROPERTY where the UIMax and ClampMax meta tags where set to 100.0 instead of 500.0 (should be 500.0 since the max value of the Advanced mode knob was set to 500.0)
- SAudioVectorscopePanelWidget: added setters to different properties so the knobs are in sync with the UMG widget properties
- Misc. code adjustments

diff --git a/Engine/Plugins/Runtime/AudioWidgets/Source/AudioWidgets/Private/AudioVectorscope.cpp b/Engine/Plugins/Runtime/AudioWidgets/Source/AudioWidgets/Private/AudioVectorscope.cpp
--- a/Engine/Plugins/Runtime/AudioWidgets/Source/AudioWidgets/Private/AudioVectorscope.cpp
+++ b/Engine/Plugins/Runtime/AudioWidgets/Source/AudioWidgets/Private/AudioVectorscope.cpp
@@ -54,7 +54,7 @@ namespace AudioWidgets
 
 		if (InPanelLayoutType == EAudioPanelLayoutType::Advanced)
 		{
-			VectorscopePanelWidget->OnTimeWindowValueChanged.AddSP(AudioSamplesDataProvider.Get(), &FWaveformAudioSamplesDataProvider::SetTimeWindow);
+			VectorscopePanelWidget->OnDisplayPersistenceValueChanged.AddSP(AudioSamplesDataProvider.Get(), &FWaveformAudioSamplesDataProvider::SetTimeWindow);
 		}
 	}
 
diff --git a/Engine/Plugins/Runtime/AudioWidgets/Source/AudioWidgets/Private/AudioVectorscopeUMG.cpp b/Engine/Plugins/Runtime/AudioWidgets/Source/AudioWidgets/Private/AudioVectorscopeUMG.cpp
--- a/Engine/Plugins/Runtime/AudioWidgets/Source/AudioWidgets/Private/AudioVectorscopeUMG.cpp
+++ b/Engine/Plugins/Runtime/AudioWidgets/Source/AudioWidgets/Private/AudioVectorscopeUMG.cpp
@@ -67,7 +67,7 @@ void UAudioVectorscope::CreateVectorscopeWidget()
 
 	if (PanelLayoutType == EAudioPanelLayoutType::Advanced)
 	{
-		VectorscopePanelWidget->OnTimeWindowValueChanged.AddSP(AudioSamplesDataProvider.Get(), &FWaveformAudioSamplesDataProvider::SetTimeWindow);
+		VectorscopePanelWidget->OnDisplayPersistenceValueChanged.AddSP(AudioSamplesDataProvider.Get(), &FWaveformAudioSamplesDataProvider::SetTimeWindow);
 	}
 }
 
@@ -107,6 +107,12 @@ void UAudioVectorscope::SynchronizeProperties()
 		}
 	}
 
+	if (AudioSamplesDataProvider.IsValid())
+	{
+		AudioSamplesDataProvider->SetTimeWindow(DisplayPersistenceMs);
+		AudioSamplesDataProvider->RequestSequenceView(TRange<double>::Inclusive(0, 1));
+	}
+
 	if (VectorscopePanelWidget.IsValid())
 	{
 		if (PanelLayoutType != VectorscopePanelWidget->GetPanelLayoutType() && AudioSamplesDataProvider.IsValid())
@@ -118,6 +124,8 @@ void UAudioVectorscope::SynchronizeProperties()
 		VectorscopePanelWidget->SetValueGridOverlayMaxNumDivisions(GridDivisions);
 
 		VectorscopePanelWidget->UpdateSequenceVectorViewerStyle(VectorscopeStyle.VectorViewerStyle);
+
+		VectorscopePanelWidget->SetDisplayPersistence(DisplayPersistenceMs);
 		VectorscopePanelWidget->SetVectorViewerScaleFactor(Scale);
 
 		VectorscopePanelWidget->SetGridVisibility(bShowGrid);
diff --git a/Engine/Plugins/Runtime/AudioWidgets/Source/AudioWidgets/Private/SAudioVectorscopePanelWidget.cpp b/Engine/Plugins/Runtime/AudioWidgets/Source/AudioWidgets/Private/SAudioVectorscopePanelWidget.cpp
--- a/Engine/Plugins/Runtime/AudioWidgets/Source/AudioWidgets/Private/SAudioVectorscopePanelWidget.cpp
+++ b/Engine/Plugins/Runtime/AudioWidgets/Source/AudioWidgets/Private/SAudioVectorscopePanelWidget.cpp
@@ -115,13 +115,13 @@ void SAudioVectorscopePanelWidget::CreateLayout()
 				.FillHeight(0.2f)
 				[
 					SNew(STextBlock)
-					.Text(LOCTEXT("Vectorscope_TimeWindow_Display_Label", "Persistence"))
+					.Text(LOCTEXT("Vectorscope_DisplayPersistence_Display_Label", "Persistence"))
 					.Justification(ETextJustify::Center)
 				]
 				+ SVerticalBox::Slot()
 				.FillHeight(0.8f)
 				[
-					TimeWindowKnob.ToSharedRef()
+					DisplayPersistenceKnob.ToSharedRef()
 				]
 			];
 
@@ -213,29 +213,29 @@ void SAudioVectorscopePanelWidget::CreateSequenceVectorViewer(const FFixedSample
 	.SequenceDrawingParams(SampledSequenceDrawingUtils::FSampledSequenceDrawingParams());
 }
 
-void SAudioVectorscopePanelWidget::CreateTimeWindowKnob()
+void SAudioVectorscopePanelWidget::CreateDisplayPersistenceKnob()
 {
 	auto OnValueChangedLambda = [this](float Value)
 	{
-		if (TimeWindowKnob.IsValid())
+		if (DisplayPersistenceKnob.IsValid())
 		{
 			if (!bIsInputWidgetTransacting)
 			{
 			#if WITH_EDITOR
 				if (GEditor)
 				{
-					GEditor->BeginTransaction(LOCTEXT("Vectorscope_TimeWindow_Knob_Changed_Msg", "Set vectorscope Time Window value."));
+					GEditor->BeginTransaction(LOCTEXT("Vectorscope_DisplayPersistence_Knob_Changed_Msg", "Set vectorscope Time Window value."));
 				}
 			#endif
 				bIsInputWidgetTransacting = true;
 			}
 
-			if (const float TimeWindowKnobValue = TimeWindowKnob->GetOutputValue(Value);
-				TimeWindowKnobValue != TimeWindowValue)
+			if (const float DisplayPersistenceKnobValue = DisplayPersistenceKnob->GetOutputValue(Value);
+				DisplayPersistenceKnobValue != DisplayPersistenceValue)
 			{
-				TimeWindowValue = TimeWindowKnobValue;
+				DisplayPersistenceValue = DisplayPersistenceKnobValue;
 
-				OnTimeWindowValueChanged.Broadcast(TimeWindowValue);
+				OnDisplayPersistenceValueChanged.Broadcast(DisplayPersistenceValue);
 			}
 		}
 	};
@@ -247,7 +247,7 @@ void SAudioVectorscopePanelWidget::CreateTimeWindowKnob()
 		#if WITH_EDITOR
 			if (GEditor)
 			{
-				GEditor->BeginTransaction(LOCTEXT("Vectorscope_TimeWindow_Knob_CaptureBegin_Msg", "Set vectorscope Time Window value."));
+				GEditor->BeginTransaction(LOCTEXT("Vectorscope_DisplayPersistence_Knob_CaptureBegin_Msg", "Set vectorscope Time Window value."));
 			}
 		#endif
 			bIsInputWidgetTransacting = true;
@@ -272,15 +272,15 @@ void SAudioVectorscopePanelWidget::CreateTimeWindowKnob()
 		}
 	};
 
-	TimeWindowKnob = SNew(SAudioRadialSlider)
+	DisplayPersistenceKnob = SNew(SAudioRadialSlider)
 	.OnMouseCaptureBegin_Lambda(OnRadialSliderMouseCaptureBeginLambda)
 	.OnMouseCaptureEnd_Lambda(OnRadialSliderMouseCaptureEndLambda)
 	.SliderValue(0.0f);
 
-	TimeWindowKnob->SetOutputRange(FVector2D(10.0f, 500.0f));
-	TimeWindowKnob->SetUnitsText(FText::FromString("ms"));
+	DisplayPersistenceKnob->SetOutputRange(DisplayPersistenceKnobOutputRange);
+	DisplayPersistenceKnob->SetUnitsText(FText::FromString("ms"));
 
-	TimeWindowKnob->OnValueChanged.BindLambda(OnValueChangedLambda);
+	DisplayPersistenceKnob->OnValueChanged.BindLambda(OnValueChangedLambda);
 }
 
 void SAudioVectorscopePanelWidget::CreateScaleKnob()
@@ -347,7 +347,7 @@ void SAudioVectorscopePanelWidget::CreateScaleKnob()
 	.OnMouseCaptureEnd_Lambda(OnRadialSliderMouseCaptureEndLambda)
 	.SliderValue(1.0f);
 
-	ScaleKnob->SetOutputRange(FVector2D(0.0f, 1.0f));
+	ScaleKnob->SetOutputRange(ScaleFactorOutputKnobRange);
 	ScaleKnob->SetShowUnitsText(false);
 
 	ScaleKnob->OnValueChanged.BindLambda(OnValueChangedLambda);
@@ -355,7 +355,7 @@ void SAudioVectorscopePanelWidget::CreateScaleKnob()
 
 void SAudioVectorscopePanelWidget::CreateVectorscopeControls()
 {
-	CreateTimeWindowKnob();
+	CreateDisplayPersistenceKnob();
 	CreateScaleKnob();
 }
 
@@ -371,9 +371,22 @@ void SAudioVectorscopePanelWidget::SetValueGridOverlayMaxNumDivisions(const uint
 	ValueGridOverlayYAxis->SetMaxDivisionParameter(InGridMaxNumDivisions);
 }
 
+void SAudioVectorscopePanelWidget::SetDisplayPersistence(const float InDisplayPersistenceInMs)
+{
+	if (DisplayPersistenceKnob.IsValid())
+	{
+		DisplayPersistenceKnob->SetSliderValue(FMath::GetMappedRangeValueUnclamped(DisplayPersistenceKnobOutputRange, { 0.0, 1.0 }, InDisplayPersistenceInMs));
+	}
+}
+
 void SAudioVectorscopePanelWidget::SetVectorViewerScaleFactor(const float InScaleFactor)
 {
 	SequenceVectorViewer->SetScaleFactor(InScaleFactor);
+
+	if (ScaleKnob.IsValid())
+	{
+		ScaleKnob->SetSliderValue(FMath::GetMappedRangeValueUnclamped(ScaleFactorOutputKnobRange, { 0.0, 1.0 }, InScaleFactor));
+	}
 }
 
 void SAudioVectorscopePanelWidget::UpdateValueGridOverlayStyle(const FSampledSequenceValueGridOverlayStyle UpdatedStyle)
diff --git a/Engine/Plugins/Runtime/AudioWidgets/Source/AudioWidgets/Public/AudioVectorscopeUMG.h b/Engine/Plugins/Runtime/AudioWidgets/Source/AudioWidgets/Public/AudioVectorscopeUMG.h
--- a/Engine/Plugins/Runtime/AudioWidgets/Source/AudioWidgets/Public/AudioVectorscopeUMG.h
+++ b/Engine/Plugins/Runtime/AudioWidgets/Source/AudioWidgets/Public/AudioVectorscopeUMG.h
@@ -60,7 +60,7 @@ public:
 	int32 GridDivisions = 2;
 
 	/** For how long the audio samples should persist in the screen (in milliseconds). */
-	UPROPERTY(EditAnywhere, BlueprintReadWrite, Category = "Vectorscope Values", meta = (UIMin = 10.0, UIMax = 100.0, ClampMin = 10.0, ClampMax = 100.0))
+	UPROPERTY(EditAnywhere, BlueprintReadWrite, Category = "Vectorscope Values", meta = (UIMin = 10.0, UIMax = 500.0, ClampMin = 10.0, ClampMax = 500.0))
 	float DisplayPersistenceMs = 60.0f;
 
 	/** The scale for the displayed audio samples. */
diff --git a/Engine/Plugins/Runtime/AudioWidgets/Source/AudioWidgets/Public/SAudioVectorscopePanelWidget.h b/Engine/Plugins/Runtime/AudioWidgets/Source/AudioWidgets/Public/SAudioVectorscopePanelWidget.h
--- a/Engine/Plugins/Runtime/AudioWidgets/Source/AudioWidgets/Public/SAudioVectorscopePanelWidget.h
+++ b/Engine/Plugins/Runtime/AudioWidgets/Source/AudioWidgets/Public/SAudioVectorscopePanelWidget.h
@@ -53,6 +53,7 @@ public:
 	void SetGridVisibility(const bool InbIsVisible);
 	void SetValueGridOverlayMaxNumDivisions(const uint32 InGridMaxNumDivisions);
 
+	void SetDisplayPersistence(const float InDisplayPersistenceInMs);
 	void SetVectorViewerScaleFactor(const float InScaleFactor);
 
 	void UpdateValueGridOverlayStyle(const FSampledSequenceValueGridOverlayStyle UpdatedValueGridOverlayStyle);
@@ -60,7 +61,7 @@ public:
 
 	EAudioPanelLayoutType GetPanelLayoutType() { return PanelLayoutType; }
 
-	FOnTimeWindowValueChanged OnTimeWindowValueChanged;
+	FOnTimeWindowValueChanged OnDisplayPersistenceValueChanged;
 
 private:
 	void CreateLayout();
@@ -76,7 +77,7 @@ private:
 	void CreateSequenceVectorViewer(const FFixedSampledSequenceView& InData, const FSampledSequenceVectorViewerStyle& VectorViewerStyle);
 
 	// Advanced panel methods
-	void CreateTimeWindowKnob();
+	void CreateDisplayPersistenceKnob();
 	void CreateScaleKnob();
 	void CreateVectorscopeControls();
 
@@ -89,17 +90,20 @@ private:
 	TSharedPtr<SSampledSequenceValueGridOverlay> ValueGridOverlayYAxis;
 
 	// Advanced panel widgets
-	TSharedPtr<SAudioRadialSlider> TimeWindowKnob;
+	TSharedPtr<SAudioRadialSlider> DisplayPersistenceKnob;
 	TSharedPtr<SAudioRadialSlider> ScaleKnob;
 
+	inline static const FVector2D DisplayPersistenceKnobOutputRange = { 10.0, 500.0 };
+	inline static const FVector2D ScaleFactorOutputKnobRange        = { 0.0, 1.0 };
+
 	EAudioPanelLayoutType PanelLayoutType = EAudioPanelLayoutType::Basic;
 
 	uint32 ValueGridMaxDivisionParameter = 2;
 
 	bool bIsInputWidgetTransacting = false;
 
-	float TimeWindowValue = 0.0f;
-	float ScaleValue      = 0.0f;
+	float DisplayPersistenceValue = 0.0f;
+	float ScaleValue = 0.0f;
 
 	float VectorscopeViewProportion = 1.0f;
 
