[AudioWidgets] Oscilloscope: Fixed multiple bugs related to UMG widget:
- UAudioOscilloscope: Added ChannelsToAnalyze UPROPERTY (it will only show if Advanced mode is selected)
- UAudioOscilloscope: Added TriggerMode UPROPERTY (it will only show if Advanced mode is selected or if the audio bus has only one channel since multi-channel triggering is not supported)
- UAudioOscilloscope: Added logic to support the 2 UPROPERTIES described above
- SAudioOscilloscopePanelWidget: added setters to different properties so the advanced mode controls are in sync with the UMG widget properties
- FWaveformAudioSamplesDataProvider: Don't reset audio buffers if the ChannelIndexToAnalyze has not changed
- Misc. code adjustments

diff --git a/Engine/Plugins/Runtime/AudioWidgets/Source/AudioWidgets/Private/AudioOscilloscopeUMG.cpp b/Engine/Plugins/Runtime/AudioWidgets/Source/AudioWidgets/Private/AudioOscilloscopeUMG.cpp
--- a/Engine/Plugins/Runtime/AudioWidgets/Source/AudioWidgets/Private/AudioOscilloscopeUMG.cpp
+++ b/Engine/Plugins/Runtime/AudioWidgets/Source/AudioWidgets/Private/AudioOscilloscopeUMG.cpp
@@ -31,16 +31,13 @@ UAudioOscilloscope::UAudioOscilloscope(const FObjectInitializer& ObjectInitializ
 
 void UAudioOscilloscope::CreateDummyOscilloscopeWidget()
 {
-	OscilloscopePanelWidget = SNew(SAudioOscilloscopePanelWidget, DummyDataView, 1)
+	OscilloscopePanelWidget = SNew(SAudioOscilloscopePanelWidget, DummyDataView, DummyNumChannels)
 	.PanelLayoutType(PanelLayoutType)
 	.PanelStyle(&OscilloscopeStyle);
 }
 
 void UAudioOscilloscope::CreateDataProvider()
 {
-	constexpr uint32 NumChannelsToProvide = 1;
-	constexpr float  MaxTimeWindowMs      = 5000.0f; // TODO alex.perez: should we expose this as a UPROPERTY?
-
 	UWorld* World = GetWorld();
 	if (!World)
 	{
@@ -53,6 +50,9 @@ void UAudioOscilloscope::CreateDataProvider()
 		return;
 	}
 
+	const uint32 NumChannelsToProvide = (PanelLayoutType == EAudioPanelLayoutType::Advanced) ? 1 : AudioBus->GetNumChannels(); // Advanced mode waveform display is based on channel selection
+	constexpr float  MaxTimeWindowMs = 5000.0f; // TODO alex.perez: should we expose this as a UPROPERTY?
+
 	AudioSamplesDataProvider = MakeShared<AudioWidgets::FWaveformAudioSamplesDataProvider>(AudioDevice.GetDeviceID(), AudioBus, NumChannelsToProvide, TimeWindowMs, MaxTimeWindowMs, AnalysisPeriodMs);
 }
 
@@ -85,10 +85,14 @@ TSharedRef<SWidget> UAudioOscilloscope::RebuildWidget()
 {
 	if (!AudioBus)
 	{
+		NumChannels = DummyNumChannels;
+
 		CreateDummyOscilloscopeWidget();
 	}
 	else
 	{
+		NumChannels = AudioBus->GetNumChannels();
+
 		CreateDataProvider();
 		CreateOscilloscopeWidget();
 	}
@@ -102,6 +106,8 @@ void UAudioOscilloscope::SynchronizeProperties()
 
 	if (!AudioBus)
 	{
+		NumChannels = DummyNumChannels;
+
 		if (AudioSamplesDataProvider.IsValid())
 		{
 			AudioSamplesDataProvider.Reset();
@@ -110,15 +116,37 @@ void UAudioOscilloscope::SynchronizeProperties()
 	}
 	else
 	{
-		if (!AudioSamplesDataProvider.IsValid() || AudioBus != AudioSamplesDataProvider->GetAudioBus())
+		if (!AudioSamplesDataProvider.IsValid() || AudioBus != AudioSamplesDataProvider->GetAudioBus() || AudioBus->GetNumChannels() != NumChannels)
 		{
+			NumChannels = AudioBus->GetNumChannels();
+
 			CreateDataProvider();
 			CreateOscilloscopeWidget();
 		}
 	}
 
+	if (PanelLayoutType == EAudioPanelLayoutType::Advanced)
+	{
+		ChannelToAnalyze = FMath::Clamp(ChannelToAnalyze, 1, NumChannels);
+	}
+
 	if (AudioSamplesDataProvider.IsValid())
 	{
+		if (PanelLayoutType == EAudioPanelLayoutType::Advanced)
+		{
+			AudioSamplesDataProvider->SetChannelToAnalyze(ChannelToAnalyze);
+		}
+
+		if (CanTriggeringBeSet())
+		{
+			AudioSamplesDataProvider->SetTriggerMode(TriggerMode);
+			AudioSamplesDataProvider->SetTriggerThreshold(TriggerThreshold);
+		}
+		else
+		{
+			AudioSamplesDataProvider->SetTriggerMode(EAudioOscilloscopeTriggerMode::None);
+		}
+
 		AudioSamplesDataProvider->SetTimeWindow(TimeWindowMs);
 		AudioSamplesDataProvider->SetAnalysisPeriod(AnalysisPeriodMs);
 		AudioSamplesDataProvider->RequestSequenceView(TRange<double>::Inclusive(0, 1));
@@ -148,8 +176,23 @@ void UAudioOscilloscope::SynchronizeProperties()
 		OscilloscopePanelWidget->SetYAxisLabelsVisibility(bShowAmplitudeLabels);
 		OscilloscopePanelWidget->SetValueGridOverlayDisplayUnit(AmplitudeGridLabelsUnit);
 
-		OscilloscopePanelWidget->SetTriggerThreshold(TriggerThreshold);
-		OscilloscopePanelWidget->SetTriggerThresholdVisibility(bShowTriggerThresholdLine);
+		if (PanelLayoutType == EAudioPanelLayoutType::Advanced)
+		{
+			OscilloscopePanelWidget->SetChannelToAnalyze(ChannelToAnalyze);
+		}
+
+		if (CanTriggeringBeSet())
+		{
+			OscilloscopePanelWidget->SetTriggerMode(TriggerMode);
+			OscilloscopePanelWidget->SetTriggerThreshold(TriggerThreshold);
+		}
+		else
+		{
+			OscilloscopePanelWidget->SetTriggerMode(EAudioOscilloscopeTriggerMode::None);
+		}
+
+		OscilloscopePanelWidget->SetTimeWindow(TimeWindowMs);
+		OscilloscopePanelWidget->SetAnalysisPeriod(AnalysisPeriodMs);
 	}
 }
 
@@ -167,6 +210,11 @@ const FText UAudioOscilloscope::GetPaletteCategory()
 }
 #endif
 
+bool UAudioOscilloscope::CanTriggeringBeSet()
+{
+	return PanelLayoutType == EAudioPanelLayoutType::Advanced || (PanelLayoutType == EAudioPanelLayoutType::Basic && NumChannels == 1);
+}
+
 void UAudioOscilloscope::StartProcessing()
 {
 	if (AudioSamplesDataProvider.IsValid())
diff --git a/Engine/Plugins/Runtime/AudioWidgets/Source/AudioWidgets/Private/SAudioOscilloscopePanelWidget.cpp b/Engine/Plugins/Runtime/AudioWidgets/Source/AudioWidgets/Private/SAudioOscilloscopePanelWidget.cpp
--- a/Engine/Plugins/Runtime/AudioWidgets/Source/AudioWidgets/Private/SAudioOscilloscopePanelWidget.cpp
+++ b/Engine/Plugins/Runtime/AudioWidgets/Source/AudioWidgets/Private/SAudioOscilloscopePanelWidget.cpp
@@ -373,6 +373,12 @@ void SAudioOscilloscopePanelWidget::CreateChannelCombobox()
 
 	auto OnChannelSelectionChanged = [this](TSharedPtr<FString> InChannel, ESelectInfo::Type InSelectInfo)
 	{
+		if (!InChannel.IsValid())
+		{
+			UE_LOG(LogAudioWidgets, Verbose, TEXT("Oscilloscope channel selection changed but the TSharedPtr that carries the channel number FString is null."));
+			return;
+		}
+
 		SelectedChannelPtr = InChannel;
 		ChannelCombobox->SetSelectedItem(SelectedChannelPtr);
 
@@ -413,6 +419,12 @@ void SAudioOscilloscopePanelWidget::CreateTriggerModeCombobox()
 
 	auto OnTriggerModeSelectionChanged = [this](TSharedPtr<EAudioOscilloscopeTriggerMode> InTriggerMode, ESelectInfo::Type InSelectInfo)
 	{
+		if (!InTriggerMode.IsValid())
+		{
+			UE_LOG(LogAudioWidgets, Verbose, TEXT("Oscilloscope trigger mode selection changed but the TSharedPtr that carries the trigger mode is null."));
+			return;
+		}
+
 		SelectedTriggerModePtr = InTriggerMode;
 		TriggerModeCombobox->SetSelectedItem(SelectedTriggerModePtr);
 
@@ -507,7 +519,7 @@ void SAudioOscilloscopePanelWidget::CreateTriggerThresholdKnob()
 
 	auto OnRadialSliderEnabledLambda = [this]()
 	{
-		return *SelectedTriggerModePtr.Get() != EAudioOscilloscopeTriggerMode::None;
+		return SelectedTriggerModePtr.IsValid() ? *SelectedTriggerModePtr.Get() != EAudioOscilloscopeTriggerMode::None : false;
 	};
 
 	TriggerThresholdKnob = SNew(SAudioRadialSlider)
@@ -518,7 +530,7 @@ void SAudioOscilloscopePanelWidget::CreateTriggerThresholdKnob()
 	.IsEnabled_Lambda(OnRadialSliderEnabledLambda); // In normalized range, range will change below
 
 	TriggerThresholdKnob->SetShowUnitsText(false);
-	TriggerThresholdKnob->SetOutputRange(FVector2D(-1.0f, 1.0f));
+	TriggerThresholdKnob->SetOutputRange(TriggerThresholdKnobOutputRange);
 }
 
 void SAudioOscilloscopePanelWidget::CreateTimeWindowKnob()
@@ -581,7 +593,7 @@ void SAudioOscilloscopePanelWidget::CreateTimeWindowKnob()
 	.OnMouseCaptureBegin_Lambda(OnRadialSliderMouseCaptureBeginLambda)
 	.OnMouseCaptureEnd_Lambda(OnRadialSliderMouseCaptureEndLambda);
 
-	TimeWindowKnob->SetOutputRange(FVector2D(10.0f, 5000.0f));
+	TimeWindowKnob->SetOutputRange(TimeWindowKnobOutputRange);
 	TimeWindowKnob->SetUnitsText(FText::FromString("ms"));
 
 	TimeWindowKnob->OnValueChanged.BindLambda(OnValueChangedLambda);
@@ -647,7 +659,7 @@ void SAudioOscilloscopePanelWidget::CreateAnalysisPeriodKnob()
 	.OnMouseCaptureBegin_Lambda(OnRadialSliderMouseCaptureBeginLambda)
 	.OnMouseCaptureEnd_Lambda(OnRadialSliderMouseCaptureEndLambda);
 
-	AnalysisPeriodKnob->SetOutputRange(FVector2D(10.0f, 1000.0f));
+	AnalysisPeriodKnob->SetOutputRange(AnalysisPeriodKnobOutputRange);
 	AnalysisPeriodKnob->SetUnitsText(FText::FromString("ms"));
 
 	AnalysisPeriodKnob->OnValueChanged.BindLambda(OnValueChangedLambda);
@@ -689,14 +701,48 @@ void SAudioOscilloscopePanelWidget::SetYAxisLabelsVisibility(const bool InbIsVis
 	ValueGridOverlay->SetHideLabels(!InbIsVisible);
 }
 
-void SAudioOscilloscopePanelWidget::SetTriggerThreshold(float InTriggerThreshold)
+void SAudioOscilloscopePanelWidget::SetChannelToAnalyze(const int32 InChannelToAnalyze)
+{
+	if (ChannelCombobox.IsValid())
+	{
+		ChannelCombobox->SetSelectedItem(MakeShareable(new FString(FString::FromInt(InChannelToAnalyze))));
+	}
+}
+
+void SAudioOscilloscopePanelWidget::SetTriggerMode(const EAudioOscilloscopeTriggerMode InTriggerMode)
+{
+	TriggerThresholdLineWidget->SetVisibility(InTriggerMode != EAudioOscilloscopeTriggerMode::None ? EVisibility::Visible : EVisibility::Hidden);
+
+	if (TriggerModeCombobox.IsValid())
+	{
+		TriggerModeCombobox->SetSelectedItem(MakeShareable(new EAudioOscilloscopeTriggerMode(InTriggerMode)));
+	}
+}
+
+void SAudioOscilloscopePanelWidget::SetTriggerThreshold(const float InTriggerThreshold)
 {
 	TriggerThresholdLineWidget->SetTriggerThreshold(InTriggerThreshold);
+
+	if (TriggerThresholdKnob.IsValid())
+	{
+		TriggerThresholdKnob->SetSliderValue(FMath::GetMappedRangeValueUnclamped(TriggerThresholdKnobOutputRange, { 0.0, 1.0 }, InTriggerThreshold));
+	}
 }
 
-void SAudioOscilloscopePanelWidget::SetTriggerThresholdVisibility(const bool InbIsVisible)
+void SAudioOscilloscopePanelWidget::SetTimeWindow(const float InTimeWindow)
 {
-	TriggerThresholdLineWidget->SetVisibility(InbIsVisible ? EVisibility::Visible : EVisibility::Hidden);
+	if (TimeWindowKnob.IsValid())
+	{
+		TimeWindowKnob->SetSliderValue(FMath::GetMappedRangeValueUnclamped(TimeWindowKnobOutputRange, { 0.0, 1.0 }, InTimeWindow));
+	}
+}
+
+void SAudioOscilloscopePanelWidget::SetAnalysisPeriod(const float InAnalysisPeriod)
+{
+	if (AnalysisPeriodKnob.IsValid())
+	{
+		AnalysisPeriodKnob->SetSliderValue(FMath::GetMappedRangeValueUnclamped(AnalysisPeriodKnobOutputRange, { 0.0, 1.0 }, InAnalysisPeriod));
+	}
 }
 
 void SAudioOscilloscopePanelWidget::UpdateSequenceRulerStyle(const FFixedSampleSequenceRulerStyle UpdatedStyle)
diff --git a/Engine/Plugins/Runtime/AudioWidgets/Source/AudioWidgets/Private/WaveformAudioSamplesDataProvider.cpp b/Engine/Plugins/Runtime/AudioWidgets/Source/AudioWidgets/Private/WaveformAudioSamplesDataProvider.cpp
--- a/Engine/Plugins/Runtime/AudioWidgets/Source/AudioWidgets/Private/WaveformAudioSamplesDataProvider.cpp
+++ b/Engine/Plugins/Runtime/AudioWidgets/Source/AudioWidgets/Private/WaveformAudioSamplesDataProvider.cpp
@@ -147,8 +147,13 @@ namespace AudioWidgets
 
 	void FWaveformAudioSamplesDataProvider::SetChannelToAnalyze(const int32 InChannel)
 	{
-		ChannelIndexToAnalyze = InChannel - 1;
-		ResetAudioBuffers();
+		const int32 ChannelIndex = InChannel - 1;
+
+		if (ChannelIndex != ChannelIndexToAnalyze)
+		{
+			ChannelIndexToAnalyze = ChannelIndex;
+			ResetAudioBuffers();
+		}
 	}
 
 	void FWaveformAudioSamplesDataProvider::SetTriggerMode(const EAudioOscilloscopeTriggerMode InTriggerMode)
diff --git a/Engine/Plugins/Runtime/AudioWidgets/Source/AudioWidgets/Public/AudioOscilloscopeUMG.h b/Engine/Plugins/Runtime/AudioWidgets/Source/AudioWidgets/Public/AudioOscilloscopeUMG.h
--- a/Engine/Plugins/Runtime/AudioWidgets/Source/AudioWidgets/Public/AudioOscilloscopeUMG.h
+++ b/Engine/Plugins/Runtime/AudioWidgets/Source/AudioWidgets/Public/AudioOscilloscopeUMG.h
@@ -2,6 +2,7 @@
 
 #pragma once
 
+#include "AudioOscilloscopeEnums.h"
 #include "AudioOscilloscopePanelStyle.h"
 #include "AudioWidgetsEnums.h"
 #include "Components/Widget.h"
@@ -82,18 +83,22 @@ public:
 	UPROPERTY(EditAnywhere, BlueprintReadWrite, Category = "Oscilloscope Values")
 	EYAxisLabelsUnit AmplitudeGridLabelsUnit = EYAxisLabelsUnit::Linear;
 
-	/** Show/Hide the trigger threshold line. */
-	UPROPERTY(EditAnywhere, BlueprintReadWrite, Category = "Oscilloscope Values")
-	bool bShowTriggerThresholdLine = false;
+	/** The trigger detection behavior. */
+	UPROPERTY(EditAnywhere, BlueprintReadWrite, Category = "Oscilloscope Values", meta = (EditCondition = "CanTriggeringBeSet()", EditConditionHides))
+	EAudioOscilloscopeTriggerMode TriggerMode = EAudioOscilloscopeTriggerMode::None;
 
 	/** The trigger threshold position in the Y axis. */
-	UPROPERTY(EditAnywhere, BlueprintReadWrite, Category = "Oscilloscope Values", meta = (UIMin = -1, UIMax = 1, ClampMin = -1, ClampMax = 1))
+	UPROPERTY(EditAnywhere, BlueprintReadWrite, Category = "Oscilloscope Values", meta = (UIMin = -1, UIMax = 1, ClampMin = -1, ClampMax = 1, EditCondition = "CanTriggeringBeSet() && TriggerMode != EAudioOscilloscopeTriggerMode::None", EditConditionHides))
 	float TriggerThreshold = 0.0f;
 
 	/** Show/Hide advanced panel layout. */
 	UPROPERTY(EditAnywhere, BlueprintReadWrite, Category = "Oscilloscope Values", meta = (DesignerRebuild = "True"))
 	EAudioPanelLayoutType PanelLayoutType = EAudioPanelLayoutType::Basic;
 
+	/** The channel to analyze with the oscilloscope (only available if PanelLayoutType is set to "Advanced"). */
+	UPROPERTY(EditAnywhere, BlueprintReadWrite, Category = "Oscilloscope Values", meta = (EditCondition = "PanelLayoutType == EAudioPanelLayoutType::Advanced", EditConditionHides))
+	int32 ChannelToAnalyze = 1;
+
 private:
 	void CreateDummyOscilloscopeWidget();
 	void CreateDataProvider();
@@ -102,6 +107,9 @@ private:
 	// UWidget overrides
 	virtual TSharedRef<SWidget> RebuildWidget() override;
 	
+	UFUNCTION()
+	bool CanTriggeringBeSet();
+
 	// The underlying audio samples data provider
 	TSharedPtr<AudioWidgets::FWaveformAudioSamplesDataProvider> AudioSamplesDataProvider;
 
@@ -113,6 +121,8 @@ private:
 	static constexpr int32 DummyMaxNumSamples = DummySampleRate * 5;
 	static constexpr int32 DummyNumChannels   = 1;
 
+	int32 NumChannels = DummyNumChannels;
+
 	TArray<float> DummyAudioSamples;
 	FFixedSampledSequenceView DummyDataView;
 };
diff --git a/Engine/Plugins/Runtime/AudioWidgets/Source/AudioWidgets/Public/SAudioOscilloscopePanelWidget.h b/Engine/Plugins/Runtime/AudioWidgets/Source/AudioWidgets/Public/SAudioOscilloscopePanelWidget.h
--- a/Engine/Plugins/Runtime/AudioWidgets/Source/AudioWidgets/Public/SAudioOscilloscopePanelWidget.h
+++ b/Engine/Plugins/Runtime/AudioWidgets/Source/AudioWidgets/Public/SAudioOscilloscopePanelWidget.h
@@ -106,8 +106,11 @@ public:
 	void SetYAxisGridVisibility(const bool InbIsVisible);
 	void SetYAxisLabelsVisibility(const bool InbIsVisible);
 
-	void SetTriggerThreshold(float InTriggerThreshold);
-	void SetTriggerThresholdVisibility(const bool InbIsVisible);
+	void SetChannelToAnalyze(const int32 InChannelToAnalyze);
+	void SetTriggerMode(const EAudioOscilloscopeTriggerMode InTriggerMode);
+	void SetTriggerThreshold(const float InTriggerThreshold);
+	void SetTimeWindow(const float InTimeWindow);
+	void SetAnalysisPeriod(const float InAnalysisPeriod);
 
 	void UpdateSequenceRulerStyle(const FFixedSampleSequenceRulerStyle UpdatedRulerStyle);
 	void UpdateValueGridOverlayStyle(const FSampledSequenceValueGridOverlayStyle UpdatedValueGridOverlayStyle);
@@ -165,6 +168,10 @@ private:
 	TSharedPtr<SAudioRadialSlider> TimeWindowKnob;
 	TSharedPtr<SAudioRadialSlider> AnalysisPeriodKnob;
 
+	inline static const FVector2D TriggerThresholdKnobOutputRange = { -1.0, 1.0 };
+	inline static const FVector2D TimeWindowKnobOutputRange       = { 10.0, 5000.0 };
+	inline static const FVector2D AnalysisPeriodKnobOutputRange   = { 10.0, 1000.0 };
+
 	EAudioPanelLayoutType PanelLayoutType = EAudioPanelLayoutType::Basic;
 
 	uint32 NumChannels = 0;
