AudioMaterialKnob Label Unit text change and value processing depending on the selected Unit
#rnx

diff --git a/Engine/Plugins/Runtime/AudioWidgets/Source/AudioWidgets/Private/AudioMaterialSlate/SAudioMaterialLabeledKnob.cpp b/Engine/Plugins/Runtime/AudioWidgets/Source/AudioWidgets/Private/AudioMaterialSlate/SAudioMaterialLabeledKnob.cpp
--- a/Engine/Plugins/Runtime/AudioWidgets/Source/AudioWidgets/Private/AudioMaterialSlate/SAudioMaterialLabeledKnob.cpp
+++ b/Engine/Plugins/Runtime/AudioWidgets/Source/AudioWidgets/Private/AudioMaterialSlate/SAudioMaterialLabeledKnob.cpp
@@ -12,7 +12,10 @@ BEGIN_SLATE_FUNCTION_BUILD_OPTIMIZATION
 void SAudioMaterialLabeledKnob::Construct(const FArguments& InArgs)
 {
 	Style = InArgs._Style;
+	
 	Owner = InArgs._Owner;
+	AudioUnitsValueType = InArgs._AudioUnitsValueType;
+
 	OnValueChanged = InArgs._OnValueChanged;
 	OnMouseCaptureEnd = InArgs._OnMouseCaptureEnd;
 	OnMouseCaptureBegin = InArgs._OnMouseCaptureBegin;
@@ -35,11 +38,13 @@ void SAudioMaterialLabeledKnob::Construct(const FArguments& InArgs)
 				Knob->SetValue(NewSliderValue);
 				OnValueChanged.ExecuteIfBound(NewSliderValue);
 				OnMouseCaptureEnd.ExecuteIfBound();
+				OnValueTextCommitted.ExecuteIfBound(NewSliderValue);
 			}
 		});
 
 	// Underlying Knob widget
 	SAssignNew(Knob, SAudioMaterialKnob)
+		.Value(ValueAttribute)
 		.Owner(InArgs._Owner)
 		.OnFloatValueChanged_Lambda([this](float Value)
 		{
@@ -67,12 +72,37 @@ void SAudioMaterialLabeledKnob::Construct(const FArguments& InArgs)
 			Label.ToSharedRef()
 		];
 
+	//Set the Unit processor from unit value type
+	switch (AudioUnitsValueType.Get())
+	{
+	case EAudioUnitsValueType::Linear:
+		if (Label.IsValid())
+		{
+			AudioUnitProcessor = MakeShared<FAudioUnitProcessor>();
+			Label.Get()->SetShowUnitsText(false);
+		}
+		break;
+	case EAudioUnitsValueType::Frequency:
+		AudioUnitProcessor = MakeShared<FFrequencyProcessor>();
+		break;
+	case EAudioUnitsValueType::Volume:
+		AudioUnitProcessor = MakeShared<FVolumeProcessor>(InArgs._bUseLinearOutput);
+		break;
+	}
+
+	if (AudioUnitProcessor.IsValid())
+	{
+		SetOutputRange(AudioUnitProcessor->GetDefaultOutputRange());
+		if (Label.IsValid())
+		{
+			Label.Get()->SetUnitsText(AudioUnitProcessor->GetUnitsText());
+		}
+	}
+
 	ChildSlot
 	[
 		VerticalLayotWidget.ToSharedRef()
 	];
-
-	SetOutputRange(OutputRange);
 }
 END_SLATE_FUNCTION_BUILD_OPTIMIZATION
 
@@ -86,22 +116,22 @@ void SAudioMaterialLabeledKnob::SetValue(float InValue)
 
 const float SAudioMaterialLabeledKnob::GetOutputValue(const float InSliderValue)
 {
-	return FMath::GetMappedRangeValueClamped(NormalizedLinearSliderRange, OutputRange, InSliderValue);
+	return AudioUnitProcessor.IsValid() ? AudioUnitProcessor->GetOutputValue(OutputRange, InSliderValue) : 0.0f;
 }
 
 const float SAudioMaterialLabeledKnob::GetSliderValue(const float OutputValue)
 {
-	return FMath::GetMappedRangeValueClamped(OutputRange, NormalizedLinearSliderRange, OutputValue);
+	return AudioUnitProcessor.IsValid() ? AudioUnitProcessor->GetSliderValue(OutputRange, OutputValue) : 0.0f;
 }
 
 const float SAudioMaterialLabeledKnob::GetOutputValueForText(const float InSliderValue)
 {
-	return GetOutputValue(InSliderValue);
+	return AudioUnitProcessor.IsValid() ? AudioUnitProcessor->GetOutputValueForText(OutputRange, InSliderValue) : 0.0f;
 }
 
 const float SAudioMaterialLabeledKnob::GetSliderValueForText(const float OutputValue)
 {
-	return GetSliderValue(OutputValue);
+	return AudioUnitProcessor.IsValid() ? AudioUnitProcessor->GetSliderValueForText(OutputRange, OutputValue) : 0.0f;
 }
 
 void SAudioMaterialLabeledKnob::SetSliderValue(float InSliderValue)
@@ -109,9 +139,12 @@ void SAudioMaterialLabeledKnob::SetSliderValue(float InSliderValue)
 	SetValue(InSliderValue);
 }
 
-void SAudioMaterialLabeledKnob::SetOutputRange(const FVector2D Range)
+void SAudioMaterialLabeledKnob::SetOutputRange(const FVector2D InRange)
 {
+	//Check the valid range from the processor
+	FVector2D Range = AudioUnitProcessor.IsValid() ? AudioUnitProcessor->GetOutputRange(InRange) : InRange;
 	OutputRange = Range;
+
 	// if Range.Y < Range.X, set Range.X to Range.Y
 	OutputRange.X = FMath::Min(Range.X, Range.Y);
 
@@ -155,7 +188,7 @@ void SAudioMaterialLabeledKnob::SetUnitsText(const FText Units)
 
 void SAudioMaterialLabeledKnob::SetUnitsTextReadOnly(const bool bIsReadOnly)
 {
-	Label->SetValueTextReadOnly(bIsReadOnly);
+	Label->SetUnitsTextReadOnly(bIsReadOnly);
 }
 
 void SAudioMaterialLabeledKnob::SetShowUnitsText(const bool bShowUnitsText)
diff --git a/Engine/Plugins/Runtime/AudioWidgets/Source/AudioWidgets/Public/AudioMaterialSlate/SAudioMaterialLabeledKnob.h b/Engine/Plugins/Runtime/AudioWidgets/Source/AudioWidgets/Public/AudioMaterialSlate/SAudioMaterialLabeledKnob.h
--- a/Engine/Plugins/Runtime/AudioWidgets/Source/AudioWidgets/Public/AudioMaterialSlate/SAudioMaterialLabeledKnob.h
+++ b/Engine/Plugins/Runtime/AudioWidgets/Source/AudioWidgets/Public/AudioMaterialSlate/SAudioMaterialLabeledKnob.h
@@ -4,6 +4,7 @@
 
 #include "AudioMaterialSlate/AudioMaterialSlateTypes.h"
 #include "Framework/SlateDelegates.h"
+#include "AudioWidgetsEnums.h"
 #include "SAudioInputWidget.h"
 #include "Styling/ISlateStyle.h"
 #include "Styling/SlateWidgetStyleAsset.h"
@@ -23,18 +24,27 @@ public:
 	SLATE_BEGIN_ARGS(SAudioMaterialLabeledKnob)
 	{}
 
+	/** A value representing the normalized linear (0 - 1) knobs value position. */
+	SLATE_ATTRIBUTE(float, Value)
+
 	/** The owner object*/
 	SLATE_ARGUMENT(TWeakObjectPtr<UObject>, Owner)
 
-	/**Value of the Knob*/
-	SLATE_ATTRIBUTE(float, Value)
+	/** The knob's ValueType. */
+	SLATE_ARGUMENT(EAudioUnitsValueType, AudioUnitsValueType)
+
+	/** Will the knob use Linear Output. This is used when ValueType is Volume */
+	SLATE_ARGUMENT(bool, bUseLinearOutput)
 
 	/** The style used to draw the knob. */
 	SLATE_STYLE_ARGUMENT(FAudioMaterialKnobStyle, Style)
 
-	/** Called when the knob's state changes. */
+	/** Called when the knob's value is changed by tuning or typing. */
 	SLATE_EVENT(FOnFloatValueChanged, OnValueChanged)
 
+	/** Called when the value is committed from label's text field */
+	SLATE_EVENT(FOnFloatValueChanged, OnValueTextCommitted)
+
 	/** Invoked when the mouse is pressed and a capture begins. */
 	SLATE_EVENT(FSimpleDelegate, OnMouseCaptureBegin)
 
@@ -59,7 +69,7 @@ public:
 	 * Set the knob's linear (0-1 normalized) value.
 	 */
 	virtual void SetSliderValue(float InSliderValue) override;
-	virtual void SetOutputRange(const FVector2D Range) override;
+	virtual void SetOutputRange(const FVector2D InRange) override;
 	virtual void SetDesiredSizeOverride(const FVector2D Size) override;
 	virtual void SetLabelBackgroundColor(FSlateColor InColor) override;
 	virtual void SetUnitsText(const FText Units) override;
@@ -72,6 +82,9 @@ public:
 	// Holds a delegate that is executed when the knob's value changes.
 	FOnFloatValueChanged OnValueChanged;
 
+	// Holds a delegate that is executed when the value is committed from label's text field
+	FOnFloatValueChanged OnValueTextCommitted;
+
 	// Holds a delegate that is executed when the mouse is pressed and a capture begins.
 	FSimpleDelegate OnMouseCaptureBegin;
 
@@ -97,6 +110,9 @@ private:
 	//Holds the knobs current Value
 	TAttribute<float> ValueAttribute = 1.0f;
 
+	// Holds the knob's unit value type
+	TAttribute<EAudioUnitsValueType> AudioUnitsValueType;
+
 	// Widget components
 	TSharedPtr<SAudioMaterialKnob> Knob;
 	TSharedPtr<SAudioTextBox> Label;
@@ -108,4 +124,7 @@ private:
 	FVector2D OutputRange = FVector2D(0.0f, 1.0f);
 	const FVector2D NormalizedLinearSliderRange = FVector2D(0.0f, 1.0f);
 
+	/**Hold the ref to the current Unit processor */
+	TSharedPtr<FAudioUnitProcessor> AudioUnitProcessor;
+
 };
diff --git a/Engine/Plugins/Runtime/AudioWidgets/Source/AudioWidgets/Public/AudioMaterialSlate/SAudioMaterialLabeledSlider.h b/Engine/Plugins/Runtime/AudioWidgets/Source/AudioWidgets/Public/AudioMaterialSlate/SAudioMaterialLabeledSlider.h
--- a/Engine/Plugins/Runtime/AudioWidgets/Source/AudioWidgets/Public/AudioMaterialSlate/SAudioMaterialLabeledSlider.h
+++ b/Engine/Plugins/Runtime/AudioWidgets/Source/AudioWidgets/Public/AudioMaterialSlate/SAudioMaterialLabeledSlider.h
@@ -48,7 +48,7 @@ public:
 	/** Called when the value is changed by slider or typing */
 	SLATE_EVENT(FOnFloatValueChanged, OnValueChanged)
 
-	/** Called when the value is committed by mouse capture ending */
+	/** Called when the value is committed from label's text field */
 	SLATE_EVENT(FOnFloatValueChanged, OnValueCommitted)
 
 	SLATE_END_ARGS()
diff --git a/Engine/Plugins/Runtime/Metasound/Source/MetasoundEditor/Private/SMetasoundGraphNode.cpp b/Engine/Plugins/Runtime/Metasound/Source/MetasoundEditor/Private/SMetasoundGraphNode.cpp
--- a/Engine/Plugins/Runtime/Metasound/Source/MetasoundEditor/Private/SMetasoundGraphNode.cpp
+++ b/Engine/Plugins/Runtime/Metasound/Source/MetasoundEditor/Private/SMetasoundGraphNode.cpp
@@ -861,6 +861,8 @@ namespace Metasound
 								SAssignNew(InputWidget, SAudioMaterialLabeledKnob)
 									.Owner(GraphMember->GetOwningGraph())
 									.OnValueChanged_Lambda(OnValueChangedLambda)
+									.AudioUnitsValueType(DefaultFloat->WidgetUnitValueType)
+									.bUseLinearOutput(DefaultFloat->VolumeWidgetUseLinearOutput)
 									.OnMouseCaptureBegin_Lambda(OnRadialSliderMouseCaptureBeginLambda)
 									.OnMouseCaptureEnd_Lambda(OnRadialSliderMouseCaptureEndLambda);
 							}
