[AudioWidgets][Oscilloscope/Vectorscope] Expose MaxTimeWindowMs and MaxDisplayPersistenceMs as UMG tweakable parameters.

diff --git a/Engine/Plugins/Runtime/AudioWidgets/Source/AudioWidgets/Private/AudioOscilloscopeUMG.cpp b/Engine/Plugins/Runtime/AudioWidgets/Source/AudioWidgets/Private/AudioOscilloscopeUMG.cpp
--- a/Engine/Plugins/Runtime/AudioWidgets/Source/AudioWidgets/Private/AudioOscilloscopeUMG.cpp
+++ b/Engine/Plugins/Runtime/AudioWidgets/Source/AudioWidgets/Private/AudioOscilloscopeUMG.cpp
@@ -51,7 +51,6 @@ void UAudioOscilloscope::CreateDataProvider()
 	}
 
 	const uint32 NumChannelsToProvide = (PanelLayoutType == EAudioPanelLayoutType::Advanced) ? 1 : AudioBus->GetNumChannels(); // Advanced mode waveform display is based on channel selection
-	constexpr float  MaxTimeWindowMs = 5000.0f; // TODO alex.perez: should we expose this as a UPROPERTY?
 
 	AudioSamplesDataProvider = MakeShared<AudioWidgets::FWaveformAudioSamplesDataProvider>(AudioDevice.GetDeviceID(), AudioBus, NumChannelsToProvide, TimeWindowMs, MaxTimeWindowMs, AnalysisPeriodMs);
 }
@@ -83,6 +82,8 @@ void UAudioOscilloscope::CreateOscilloscopeWidget()
 
 TSharedRef<SWidget> UAudioOscilloscope::RebuildWidget()
 {
+	TimeWindowMs = FMath::Clamp(TimeWindowMs, 10.0f, MaxTimeWindowMs);
+
 	if (!AudioBus)
 	{
 		NumChannels = DummyNumChannels;
@@ -104,6 +105,8 @@ void UAudioOscilloscope::SynchronizeProperties()
 {
 	Super::SynchronizeProperties();
 
+	TimeWindowMs = FMath::Clamp(TimeWindowMs, 10.0f, MaxTimeWindowMs);
+
 	if (!AudioBus)
 	{
 		NumChannels = DummyNumChannels;
@@ -132,6 +135,11 @@ void UAudioOscilloscope::SynchronizeProperties()
 
 	if (AudioSamplesDataProvider.IsValid())
 	{
+		if (MaxTimeWindowMs != AudioSamplesDataProvider->GetMaxTimeWindowMs())
+		{
+			AudioSamplesDataProvider->SetMaxTimeWindowMs(MaxTimeWindowMs);
+		}
+
 		if (PanelLayoutType == EAudioPanelLayoutType::Advanced)
 		{
 			AudioSamplesDataProvider->SetChannelToAnalyze(ChannelToAnalyze);
@@ -149,6 +157,7 @@ void UAudioOscilloscope::SynchronizeProperties()
 
 		AudioSamplesDataProvider->SetTimeWindow(TimeWindowMs);
 		AudioSamplesDataProvider->SetAnalysisPeriod(AnalysisPeriodMs);
+
 		AudioSamplesDataProvider->RequestSequenceView(TRange<double>::Inclusive(0, 1));
 	}
 
@@ -176,6 +185,8 @@ void UAudioOscilloscope::SynchronizeProperties()
 		OscilloscopePanelWidget->SetYAxisLabelsVisibility(bShowAmplitudeLabels);
 		OscilloscopePanelWidget->SetValueGridOverlayDisplayUnit(AmplitudeGridLabelsUnit);
 
+		OscilloscopePanelWidget->SetMaxTimeWindowMs(MaxTimeWindowMs);
+
 		if (PanelLayoutType == EAudioPanelLayoutType::Advanced)
 		{
 			OscilloscopePanelWidget->SetChannelToAnalyze(ChannelToAnalyze);
diff --git a/Engine/Plugins/Runtime/AudioWidgets/Source/AudioWidgets/Private/AudioVectorscopeUMG.cpp b/Engine/Plugins/Runtime/AudioWidgets/Source/AudioWidgets/Private/AudioVectorscopeUMG.cpp
--- a/Engine/Plugins/Runtime/AudioWidgets/Source/AudioWidgets/Private/AudioVectorscopeUMG.cpp
+++ b/Engine/Plugins/Runtime/AudioWidgets/Source/AudioWidgets/Private/AudioVectorscopeUMG.cpp
@@ -33,8 +33,6 @@ void UAudioVectorscope::CreateDummyVectorscopeWidget()
 
 void UAudioVectorscope::CreateDataProvider()
 {
-	constexpr float MaxDisplayPersistenceMs = 500.0f; // TODO alex.perez: should we expose this as a UPROPERTY?
-
 	UWorld* World = GetWorld();
 	if (!World)
 	{
@@ -73,6 +71,8 @@ void UAudioVectorscope::CreateVectorscopeWidget()
 
 TSharedRef<SWidget> UAudioVectorscope::RebuildWidget()
 {
+	DisplayPersistenceMs = FMath::Clamp(DisplayPersistenceMs, 10.0f, MaxDisplayPersistenceMs);
+
 	if (!AudioBus)
 	{
 		CreateDummyVectorscopeWidget();
@@ -90,6 +90,8 @@ void UAudioVectorscope::SynchronizeProperties()
 {
 	Super::SynchronizeProperties();
 
+	DisplayPersistenceMs = FMath::Clamp(DisplayPersistenceMs, 10.0f, MaxDisplayPersistenceMs);
+
 	if (!AudioBus)
 	{
 		if (AudioSamplesDataProvider.IsValid())
@@ -109,7 +111,13 @@ void UAudioVectorscope::SynchronizeProperties()
 
 	if (AudioSamplesDataProvider.IsValid())
 	{
+		if (MaxDisplayPersistenceMs != AudioSamplesDataProvider->GetMaxTimeWindowMs())
+		{
+			AudioSamplesDataProvider->SetMaxTimeWindowMs(MaxDisplayPersistenceMs);
+		}
+
 		AudioSamplesDataProvider->SetTimeWindow(DisplayPersistenceMs);
+
 		AudioSamplesDataProvider->RequestSequenceView(TRange<double>::Inclusive(0, 1));
 	}
 
@@ -125,6 +133,8 @@ void UAudioVectorscope::SynchronizeProperties()
 
 		VectorscopePanelWidget->UpdateSequenceVectorViewerStyle(VectorscopeStyle.VectorViewerStyle);
 
+		VectorscopePanelWidget->SetMaxDisplayPersistence(MaxDisplayPersistenceMs);
+
 		VectorscopePanelWidget->SetDisplayPersistence(DisplayPersistenceMs);
 		VectorscopePanelWidget->SetVectorViewerScaleFactor(Scale);
 
diff --git a/Engine/Plugins/Runtime/AudioWidgets/Source/AudioWidgets/Private/SAudioOscilloscopePanelWidget.cpp b/Engine/Plugins/Runtime/AudioWidgets/Source/AudioWidgets/Private/SAudioOscilloscopePanelWidget.cpp
--- a/Engine/Plugins/Runtime/AudioWidgets/Source/AudioWidgets/Private/SAudioOscilloscopePanelWidget.cpp
+++ b/Engine/Plugins/Runtime/AudioWidgets/Source/AudioWidgets/Private/SAudioOscilloscopePanelWidget.cpp
@@ -701,6 +701,16 @@ void SAudioOscilloscopePanelWidget::SetYAxisLabelsVisibility(const bool InbIsVis
 	ValueGridOverlay->SetHideLabels(!InbIsVisible);
 }
 
+void SAudioOscilloscopePanelWidget::SetMaxTimeWindowMs(const float InMaxTimeWindowMs)
+{
+	TimeWindowKnobOutputRange.Y = InMaxTimeWindowMs;
+
+	if (TimeWindowKnob.IsValid())
+	{
+		TimeWindowKnob->SetOutputRange(TimeWindowKnobOutputRange);
+	}
+}
+
 void SAudioOscilloscopePanelWidget::SetChannelToAnalyze(const int32 InChannelToAnalyze)
 {
 	if (ChannelCombobox.IsValid())
diff --git a/Engine/Plugins/Runtime/AudioWidgets/Source/AudioWidgets/Private/SAudioVectorscopePanelWidget.cpp b/Engine/Plugins/Runtime/AudioWidgets/Source/AudioWidgets/Private/SAudioVectorscopePanelWidget.cpp
--- a/Engine/Plugins/Runtime/AudioWidgets/Source/AudioWidgets/Private/SAudioVectorscopePanelWidget.cpp
+++ b/Engine/Plugins/Runtime/AudioWidgets/Source/AudioWidgets/Private/SAudioVectorscopePanelWidget.cpp
@@ -371,6 +371,16 @@ void SAudioVectorscopePanelWidget::SetValueGridOverlayMaxNumDivisions(const uint
 	ValueGridOverlayYAxis->SetMaxDivisionParameter(InGridMaxNumDivisions);
 }
 
+void SAudioVectorscopePanelWidget::SetMaxDisplayPersistence(const float InMaxDisplayPersistenceInMs)
+{
+	DisplayPersistenceKnobOutputRange.Y = InMaxDisplayPersistenceInMs;
+
+	if (DisplayPersistenceKnob.IsValid())
+	{
+		DisplayPersistenceKnob->SetOutputRange(DisplayPersistenceKnobOutputRange);
+	}
+}
+
 void SAudioVectorscopePanelWidget::SetDisplayPersistence(const float InDisplayPersistenceInMs)
 {
 	if (DisplayPersistenceKnob.IsValid())
diff --git a/Engine/Plugins/Runtime/AudioWidgets/Source/AudioWidgets/Private/WaveformAudioSamplesDataProvider.cpp b/Engine/Plugins/Runtime/AudioWidgets/Source/AudioWidgets/Private/WaveformAudioSamplesDataProvider.cpp
--- a/Engine/Plugins/Runtime/AudioWidgets/Source/AudioWidgets/Private/WaveformAudioSamplesDataProvider.cpp
+++ b/Engine/Plugins/Runtime/AudioWidgets/Source/AudioWidgets/Private/WaveformAudioSamplesDataProvider.cpp
@@ -76,12 +76,14 @@ namespace AudioWidgets
 		}
 
 		// Init audio buffers
-		TimeWindowMaxTimeSamples = FMath::RoundToInt(InMaxTimeWindowMs / 1000.0f * NumChannelsToProvide * SampleRate);
+		MaxTimeWindowMs = InMaxTimeWindowMs;
 
-		AudioSamplesForView.Init(0.0f, TimeWindowMaxTimeSamples);
+		const float MaxTimeWindowSamples = FMath::RoundToInt(MaxTimeWindowMs / 1000.0f * NumChannelsToProvide * SampleRate);
+
+		AudioSamplesForView.Init(0.0f, MaxTimeWindowSamples);
 		DataView = FFixedSampledSequenceView{ MakeArrayView(AudioSamplesForView.GetData(), AudioSamplesForView.Num()), NumChannelsToProvide, SampleRate };
 
-		AudioSamplesCircularBuffer.SetCapacity(TimeWindowMaxTimeSamples * 2); // Twice of the amount needed for display 
+		AudioSamplesCircularBuffer.SetCapacity(MaxTimeWindowSamples * 2); // Twice of the amount needed for display 
 
 		SetTimeWindow(InTimeWindowMs);
 		SetAnalysisPeriod(InAnalysisPeriodMs);
@@ -101,10 +103,12 @@ namespace AudioWidgets
 
 	void FWaveformAudioSamplesDataProvider::ResetAudioBuffers()
 	{
-		AudioSamplesCircularBuffer.Reset(TimeWindowMaxTimeSamples * 2);
-		AudioSamplesForView.Init(0.0f, AudioSamplesForView.Num());
+		const float MaxTimeWindowSamples = FMath::RoundToInt(MaxTimeWindowMs / 1000.0f * NumChannelsToProvide * SampleRate);
 
-		AudioSamplesCircularBuffer.SetNum(TimeWindowSamples);
+		AudioSamplesCircularBuffer.Reset(MaxTimeWindowSamples * 2);
+		AudioSamplesForView.Init(0.0f, MaxTimeWindowSamples);
+
+		AudioSamplesCircularBuffer.SetNum(FMath::Clamp(TimeWindowSamples, 0.0f, MaxTimeWindowSamples));
 	}
 
 	void FWaveformAudioSamplesDataProvider::StartProcessing()
@@ -145,6 +149,12 @@ namespace AudioWidgets
 		}
 	}
 
+	void FWaveformAudioSamplesDataProvider::SetMaxTimeWindowMs(const float InMaxTimeWindowMs)
+	{
+		MaxTimeWindowMs = InMaxTimeWindowMs;
+		ResetAudioBuffers();
+	}
+
 	void FWaveformAudioSamplesDataProvider::SetChannelToAnalyze(const int32 InChannel)
 	{
 		const int32 ChannelIndex = InChannel - 1;
diff --git a/Engine/Plugins/Runtime/AudioWidgets/Source/AudioWidgets/Public/AudioOscilloscopeUMG.h b/Engine/Plugins/Runtime/AudioWidgets/Source/AudioWidgets/Public/AudioOscilloscopeUMG.h
--- a/Engine/Plugins/Runtime/AudioWidgets/Source/AudioWidgets/Public/AudioOscilloscopeUMG.h
+++ b/Engine/Plugins/Runtime/AudioWidgets/Source/AudioWidgets/Public/AudioOscilloscopeUMG.h
@@ -48,55 +48,59 @@ public:
 	void StopProcessing();
 
 	/** The oscilloscope panel style */
-	UPROPERTY(EditAnywhere, BlueprintReadWrite, Category="Style", meta=(DisplayName="Style"))
+	UPROPERTY(EditAnywhere, BlueprintReadOnly, Category="Style", meta=(DisplayName="Style"))
 	FAudioOscilloscopePanelStyle OscilloscopeStyle;
 
 	/** The audio bus used to obtain audio samples for the oscilloscope */
-	UPROPERTY(EditAnywhere, BlueprintReadWrite, Category = "Oscilloscope Values", meta = (DesignerRebuild = "True"))
+	UPROPERTY(EditAnywhere, BlueprintReadOnly, Category = "Oscilloscope Values", meta = (DesignerRebuild = "True"))
 	TObjectPtr<UAudioBus> AudioBus = nullptr;
 
+	/** The max time window in milliseconds. */
+	UPROPERTY(EditAnywhere, BlueprintReadOnly, Category = "Oscilloscope Values", meta = (UIMin = 10.0, UIMax = 5000.0, ClampMin = 10.0, ClampMax = 5000.0))
+	float MaxTimeWindowMs = 5000.0f;
+
 	/** The time window in milliseconds. */
-	UPROPERTY(EditAnywhere, BlueprintReadWrite, Category = "Oscilloscope Values", meta = (UIMin = 10, UIMax = 5000, ClampMin = 10, ClampMax = 5000))
+	UPROPERTY(EditAnywhere, BlueprintReadOnly, Category = "Oscilloscope Values", meta = (UIMin = 10.0, ClampMin = 10.0))
 	float TimeWindowMs = 10.0f;
 
 	/** The analysis period in milliseconds. */
-	UPROPERTY(EditAnywhere, BlueprintReadWrite, Category = "Oscilloscope Values", meta = (UIMin = 10, UIMax = 1000, ClampMin = 10, ClampMax = 1000))
+	UPROPERTY(EditAnywhere, BlueprintReadOnly, Category = "Oscilloscope Values", meta = (UIMin = 10, UIMax = 1000, ClampMin = 10, ClampMax = 1000))
 	float AnalysisPeriodMs = 10.0f;
 
 	/** Show/Hide the time grid. */
-	UPROPERTY(EditAnywhere, BlueprintReadWrite, Category = "Oscilloscope Values")
+	UPROPERTY(EditAnywhere, BlueprintReadOnly, Category = "Oscilloscope Values")
 	bool bShowTimeGrid = true;
 
 	/** Define the time grid labels unit. */
-	UPROPERTY(EditAnywhere, BlueprintReadWrite, Category = "Oscilloscope Values")
+	UPROPERTY(EditAnywhere, BlueprintReadOnly, Category = "Oscilloscope Values")
 	EXAxisLabelsUnit TimeGridLabelsUnit = EXAxisLabelsUnit::Samples;
 
 	/** Show/Hide the amplitude grid. */
-	UPROPERTY(EditAnywhere, BlueprintReadWrite, Category = "Oscilloscope Values")
+	UPROPERTY(EditAnywhere, BlueprintReadOnly, Category = "Oscilloscope Values")
 	bool bShowAmplitudeGrid = true;
 
 	/** Show/Hide the amplitude labels. */
-	UPROPERTY(EditAnywhere, BlueprintReadWrite, Category = "Oscilloscope Values")
+	UPROPERTY(EditAnywhere, BlueprintReadOnly, Category = "Oscilloscope Values")
 	bool bShowAmplitudeLabels = true;
 
 	/** Define the amplitude grid labels unit. */
-	UPROPERTY(EditAnywhere, BlueprintReadWrite, Category = "Oscilloscope Values")
+	UPROPERTY(EditAnywhere, BlueprintReadOnly, Category = "Oscilloscope Values")
 	EYAxisLabelsUnit AmplitudeGridLabelsUnit = EYAxisLabelsUnit::Linear;
 
 	/** The trigger detection behavior. */
-	UPROPERTY(EditAnywhere, BlueprintReadWrite, Category = "Oscilloscope Values", meta = (EditCondition = "CanTriggeringBeSet()", EditConditionHides))
+	UPROPERTY(EditAnywhere, BlueprintReadOnly, Category = "Oscilloscope Values", meta = (EditCondition = "CanTriggeringBeSet()", EditConditionHides))
 	EAudioOscilloscopeTriggerMode TriggerMode = EAudioOscilloscopeTriggerMode::None;
 
 	/** The trigger threshold position in the Y axis. */
-	UPROPERTY(EditAnywhere, BlueprintReadWrite, Category = "Oscilloscope Values", meta = (UIMin = -1, UIMax = 1, ClampMin = -1, ClampMax = 1, EditCondition = "CanTriggeringBeSet() && TriggerMode != EAudioOscilloscopeTriggerMode::None", EditConditionHides))
+	UPROPERTY(EditAnywhere, BlueprintReadOnly, Category = "Oscilloscope Values", meta = (UIMin = -1, UIMax = 1, ClampMin = -1, ClampMax = 1, EditCondition = "CanTriggeringBeSet() && TriggerMode != EAudioOscilloscopeTriggerMode::None", EditConditionHides))
 	float TriggerThreshold = 0.0f;
 
 	/** Show/Hide advanced panel layout. */
-	UPROPERTY(EditAnywhere, BlueprintReadWrite, Category = "Oscilloscope Values", meta = (DesignerRebuild = "True"))
+	UPROPERTY(EditAnywhere, BlueprintReadOnly, Category = "Oscilloscope Values", meta = (DesignerRebuild = "True"))
 	EAudioPanelLayoutType PanelLayoutType = EAudioPanelLayoutType::Basic;
 
 	/** The channel to analyze with the oscilloscope (only available if PanelLayoutType is set to "Advanced"). */
-	UPROPERTY(EditAnywhere, BlueprintReadWrite, Category = "Oscilloscope Values", meta = (EditCondition = "PanelLayoutType == EAudioPanelLayoutType::Advanced", EditConditionHides))
+	UPROPERTY(EditAnywhere, BlueprintReadOnly, Category = "Oscilloscope Values", meta = (EditCondition = "PanelLayoutType == EAudioPanelLayoutType::Advanced", EditConditionHides))
 	int32 ChannelToAnalyze = 1;
 
 private:
diff --git a/Engine/Plugins/Runtime/AudioWidgets/Source/AudioWidgets/Public/AudioVectorscopeUMG.h b/Engine/Plugins/Runtime/AudioWidgets/Source/AudioWidgets/Public/AudioVectorscopeUMG.h
--- a/Engine/Plugins/Runtime/AudioWidgets/Source/AudioWidgets/Public/AudioVectorscopeUMG.h
+++ b/Engine/Plugins/Runtime/AudioWidgets/Source/AudioWidgets/Public/AudioVectorscopeUMG.h
@@ -44,31 +44,35 @@ public:
 	void StopProcessing();
 
 	/** The vectorscope panel style */
-	UPROPERTY(EditAnywhere, BlueprintReadWrite, Category="Style", meta=(DisplayName="Style"))
+	UPROPERTY(EditAnywhere, BlueprintReadOnly, Category="Style", meta=(DisplayName="Style"))
 	FAudioVectorscopePanelStyle VectorscopeStyle;
 
 	/** The audio bus used to obtain audio samples for the vectorscope */
-	UPROPERTY(EditAnywhere, BlueprintReadWrite, Category = "Vectorscope Values", meta = (DesignerRebuild = "True"))
+	UPROPERTY(EditAnywhere, BlueprintReadOnly, Category = "Vectorscope Values", meta = (DesignerRebuild = "True"))
 	TObjectPtr<UAudioBus> AudioBus = nullptr;
 
 	/** Show/Hide the vectorscope grid. */
-	UPROPERTY(EditAnywhere, BlueprintReadWrite, Category = "Vectorscope Values")
+	UPROPERTY(EditAnywhere, BlueprintReadOnly, Category = "Vectorscope Values")
 	bool bShowGrid = true;
 
 	/** The number of grid divisions. */
-	UPROPERTY(EditAnywhere, BlueprintReadWrite, Category = "Vectorscope Values", meta = (UIMin = 1, UIMax = 6, ClampMin = 1, ClampMax = 6))
+	UPROPERTY(EditAnywhere, BlueprintReadOnly, Category = "Vectorscope Values", meta = (UIMin = 1, UIMax = 6, ClampMin = 1, ClampMax = 6))
 	int32 GridDivisions = 2;
 
+	/** The max where the audio samples should persist in the screen (in milliseconds). */
+	UPROPERTY(EditAnywhere, BlueprintReadOnly, Category = "Vectorscope Values", meta = (UIMin = 10.0, UIMax = 500.0, ClampMin = 10.0, ClampMax = 500.0))
+	float MaxDisplayPersistenceMs = 500.0f;
+
 	/** For how long the audio samples should persist in the screen (in milliseconds). */
-	UPROPERTY(EditAnywhere, BlueprintReadWrite, Category = "Vectorscope Values", meta = (UIMin = 10.0, UIMax = 500.0, ClampMin = 10.0, ClampMax = 500.0))
+	UPROPERTY(EditAnywhere, BlueprintReadOnly, Category = "Vectorscope Values", meta = (UIMin = 10.0, ClampMin = 10.0))
 	float DisplayPersistenceMs = 60.0f;
 
 	/** The scale for the displayed audio samples. */
-	UPROPERTY(EditAnywhere, BlueprintReadWrite, Category = "Vectorscope Values", meta = (UIMin = 0.0, UIMax = 1.0, ClampMin = 0.0, ClampMax = 1.0))
+	UPROPERTY(EditAnywhere, BlueprintReadOnly, Category = "Vectorscope Values", meta = (UIMin = 0.0, UIMax = 1.0, ClampMin = 0.0, ClampMax = 1.0))
 	float Scale = 1.0f;
 
 	/** Show/Hide advanced panel layout. */
-	UPROPERTY(EditAnywhere, BlueprintReadWrite, Category = "Vectorscope Values", meta = (DesignerRebuild = "True"))
+	UPROPERTY(EditAnywhere, BlueprintReadOnly, Category = "Vectorscope Values", meta = (DesignerRebuild = "True"))
 	EAudioPanelLayoutType PanelLayoutType = EAudioPanelLayoutType::Basic;
 
 private:
diff --git a/Engine/Plugins/Runtime/AudioWidgets/Source/AudioWidgets/Public/SAudioOscilloscopePanelWidget.h b/Engine/Plugins/Runtime/AudioWidgets/Source/AudioWidgets/Public/SAudioOscilloscopePanelWidget.h
--- a/Engine/Plugins/Runtime/AudioWidgets/Source/AudioWidgets/Public/SAudioOscilloscopePanelWidget.h
+++ b/Engine/Plugins/Runtime/AudioWidgets/Source/AudioWidgets/Public/SAudioOscilloscopePanelWidget.h
@@ -106,6 +106,8 @@ public:
 	void SetYAxisGridVisibility(const bool InbIsVisible);
 	void SetYAxisLabelsVisibility(const bool InbIsVisible);
 
+	void SetMaxTimeWindowMs(const float InMaxTimeWindowMs);
+
 	void SetChannelToAnalyze(const int32 InChannelToAnalyze);
 	void SetTriggerMode(const EAudioOscilloscopeTriggerMode InTriggerMode);
 	void SetTriggerThreshold(const float InTriggerThreshold);
@@ -169,9 +171,10 @@ private:
 	TSharedPtr<SAudioRadialSlider> AnalysisPeriodKnob;
 
 	inline static const FVector2D TriggerThresholdKnobOutputRange = { -1.0, 1.0 };
-	inline static const FVector2D TimeWindowKnobOutputRange       = { 10.0, 5000.0 };
 	inline static const FVector2D AnalysisPeriodKnobOutputRange   = { 10.0, 1000.0 };
 
+	FVector2D TimeWindowKnobOutputRange = { 10.0, 5000.0 };
+
 	EAudioPanelLayoutType PanelLayoutType = EAudioPanelLayoutType::Basic;
 
 	uint32 NumChannels = 0;
diff --git a/Engine/Plugins/Runtime/AudioWidgets/Source/AudioWidgets/Public/SAudioVectorscopePanelWidget.h b/Engine/Plugins/Runtime/AudioWidgets/Source/AudioWidgets/Public/SAudioVectorscopePanelWidget.h
--- a/Engine/Plugins/Runtime/AudioWidgets/Source/AudioWidgets/Public/SAudioVectorscopePanelWidget.h
+++ b/Engine/Plugins/Runtime/AudioWidgets/Source/AudioWidgets/Public/SAudioVectorscopePanelWidget.h
@@ -53,6 +53,8 @@ public:
 	void SetGridVisibility(const bool InbIsVisible);
 	void SetValueGridOverlayMaxNumDivisions(const uint32 InGridMaxNumDivisions);
 
+	void SetMaxDisplayPersistence(const float InMaxDisplayPersistenceInMs);
+
 	void SetDisplayPersistence(const float InDisplayPersistenceInMs);
 	void SetVectorViewerScaleFactor(const float InScaleFactor);
 
@@ -93,8 +95,8 @@ private:
 	TSharedPtr<SAudioRadialSlider> DisplayPersistenceKnob;
 	TSharedPtr<SAudioRadialSlider> ScaleKnob;
 
-	inline static const FVector2D DisplayPersistenceKnobOutputRange = { 10.0, 500.0 };
-	inline static const FVector2D ScaleFactorOutputKnobRange        = { 0.0, 1.0 };
+	FVector2D DisplayPersistenceKnobOutputRange = { 10.0, 500.0 };
+	inline static const FVector2D ScaleFactorOutputKnobRange = { 0.0, 1.0 };
 
 	EAudioPanelLayoutType PanelLayoutType = EAudioPanelLayoutType::Basic;
 
diff --git a/Engine/Plugins/Runtime/AudioWidgets/Source/AudioWidgets/Public/WaveformAudioSamplesDataProvider.h b/Engine/Plugins/Runtime/AudioWidgets/Source/AudioWidgets/Public/WaveformAudioSamplesDataProvider.h
--- a/Engine/Plugins/Runtime/AudioWidgets/Source/AudioWidgets/Public/WaveformAudioSamplesDataProvider.h
+++ b/Engine/Plugins/Runtime/AudioWidgets/Source/AudioWidgets/Public/WaveformAudioSamplesDataProvider.h
@@ -42,9 +42,12 @@ namespace AudioWidgets
 		void StartProcessing();
 		void StopProcessing();
 
-		FFixedSampledSequenceView GetDataView() { return DataView; };
-		uint32 GetNumChannels() { return NumChannels; }
-		const UAudioBus* GetAudioBus() { return AudioBus; }
+		FFixedSampledSequenceView GetDataView() const { return DataView; };
+		uint32 GetNumChannels() const { return NumChannels; }
+		const UAudioBus* GetAudioBus() const { return AudioBus; }
+
+		float GetMaxTimeWindowMs() const { return MaxTimeWindowMs; }
+		void SetMaxTimeWindowMs(const float InMaxTimeWindowMs);
 
 		void SetChannelToAnalyze(const int32 InChannel);
 		void SetTriggerMode(const EAudioOscilloscopeTriggerMode InTriggerMode);
@@ -66,7 +69,7 @@ namespace AudioWidgets
 		uint32 NumChannels          = 0;
 		uint32 SampleRate           = 0;
 
-		uint32 TimeWindowMaxTimeSamples = 0;
+		float MaxTimeWindowMs = 0.0f;
 
 		const Audio::FMixerDevice* MixerDevice = nullptr;
 
