UE-211580 - Default material for Slate Widget Styles that are used for AudioMaterialWidgets.
#rnx

diff --git a/Engine/Plugins/Runtime/AudioWidgets/Source/AudioWidgets/Private/AudioMaterialSlate/AudioMaterialSlateTypes.cpp b/Engine/Plugins/Runtime/AudioWidgets/Source/AudioWidgets/Private/AudioMaterialSlate/AudioMaterialSlateTypes.cpp
--- a/Engine/Plugins/Runtime/AudioWidgets/Source/AudioWidgets/Private/AudioMaterialSlate/AudioMaterialSlateTypes.cpp
+++ b/Engine/Plugins/Runtime/AudioWidgets/Source/AudioWidgets/Private/AudioMaterialSlate/AudioMaterialSlateTypes.cpp
@@ -4,6 +4,13 @@
 #include "AudioMaterialSlate/AudioMaterialSlateTypes.h"
 #include "Styling/StyleDefaults.h"
 
+namespace AudioMaterialSlateTypesPrivate
+{
+	#define PLUGIN_BASE_DIR FString("/AudioWidgets/AudioMaterialSlate/")
+}
+
+using namespace AudioMaterialSlateTypesPrivate;
+
 FAudioMaterialWidgetStyle::FAudioMaterialWidgetStyle()
 	:DesiredSize(32.f, 32.f)
 {
@@ -28,7 +35,10 @@ FAudioMaterialButtonStyle::FAudioMaterialButtonStyle()
 	, ButtonPressedShadowColor(FLinearColor::Gray)
 	, ButtonPressedOutlineColor(FLinearColor::Blue)
 {
-	DesiredSize = FVector2f(128.f, 128.f);
+	FString Path = PLUGIN_BASE_DIR + "MI_AudioMaterialButton.MI_AudioMaterialButton";
+	Material = LoadObject<UMaterialInterface>(nullptr, *Path);
+
+	DesiredSize = FVector2f(128.f, 128.f);	
 }
 
 const FName FAudioMaterialButtonStyle::TypeName(TEXT("FAudioMaterialButtonStyle"));
@@ -51,6 +61,9 @@ FAudioMaterialSliderStyle::FAudioMaterialSliderStyle()
 	, HandleOutlineColor(FLinearColor(0.15f, 0.15f, 0.15f, 1.f))
 	, TextBoxStyle(FAudioTextBoxStyle::GetDefault())
 {
+	FString Path = PLUGIN_BASE_DIR + "MI_AudioMaterialRoundedSlider.MI_AudioMaterialRoundedSlider";
+	Material = LoadObject<UMaterialInterface>(nullptr, *Path);
+
 	DesiredSize = FVector2f(30.f, 250.f);
 }
 
@@ -78,6 +91,9 @@ FAudioMaterialKnobStyle::FAudioMaterialKnobStyle()
 	, KnobBarFillMaxColor(FLinearColor::White)
 	, KnobBarFillTintColor(FLinearColor::White)
 {
+	FString Path = PLUGIN_BASE_DIR + "MI_AudioMaterialKnob.MI_AudioMaterialKnob";
+	Material = LoadObject<UMaterialInterface>(nullptr, *Path);
+
 	DesiredSize = FVector2f(128.f,128.f);
 }
 
@@ -109,6 +125,9 @@ FAudioMaterialMeterStyle::FAudioMaterialMeterStyle()
 	, DecibelsPerHash(5)
 	, Font(FStyleDefaults::GetFontInfo(5))
 {
+	FString Path = PLUGIN_BASE_DIR + "MI_AudioMaterialMeter.MI_AudioMaterialMeter";
+	Material = LoadObject<UMaterialInterface>(nullptr, *Path);
+
 	DesiredSize = FVector2f(25.f, 512.f);
 }
 
@@ -130,6 +149,9 @@ FAudioMaterialEnvelopeStyle::FAudioMaterialEnvelopeStyle()
 	, BackgroundColor(FLinearColor::Black)
 	, OutlineColor(FLinearColor::Gray)
 {
+	FString Path = PLUGIN_BASE_DIR + "MI_AudioMaterialEnvelope_ADSR.MI_AudioMaterialEnvelope_ADSR";
+	Material = LoadObject<UMaterialInterface>(nullptr, *Path);
+
 	DesiredSize = FVector2f(256.f, 256.f);
 }
 
diff --git a/Engine/Plugins/Runtime/AudioWidgets/Source/AudioWidgets/Private/AudioWidgetsStyle.cpp b/Engine/Plugins/Runtime/AudioWidgets/Source/AudioWidgets/Private/AudioWidgetsStyle.cpp
--- a/Engine/Plugins/Runtime/AudioWidgets/Source/AudioWidgets/Private/AudioWidgetsStyle.cpp
+++ b/Engine/Plugins/Runtime/AudioWidgets/Source/AudioWidgets/Private/AudioWidgetsStyle.cpp
@@ -129,43 +129,28 @@ FAudioWidgetsStyle::FAudioWidgetsStyle()
 	/**
 	**AudioMaterialKnob Style
 	*/
-	UMaterialInterface* KnobDefaultMaterial = LoadObject<UMaterialInterface>(nullptr, TEXT("/AudioWidgets/AudioMaterialSlate/MI_AudioMaterialKnob.MI_AudioMaterialKnob"));
-
-	Set("AudioMaterialKnob.Style", FAudioMaterialKnobStyle()
-		.SetMaterial(KnobDefaultMaterial));
+	Set("AudioMaterialKnob.Style", FAudioMaterialKnobStyle());
 
 	/**
 	**AudioMaterialMeter Style
 	*/
-	UMaterialInterface* MeterDefaultMaterial = LoadObject<UMaterialInterface>(nullptr, TEXT("/AudioWidgets/AudioMaterialSlate/MI_AudioMaterialMeter.MI_AudioMaterialMeter"));
-
-	Set("AudioMaterialMeter.Style", FAudioMaterialMeterStyle()
-		.SetMaterial(MeterDefaultMaterial));
+	Set("AudioMaterialMeter.Style", FAudioMaterialMeterStyle());
 
 	/**
 	**AudioMaterialEnvelope Style
 	*/
-	UMaterialInterface* EnvelopeDefaultMaterial = LoadObject<UMaterialInterface>(nullptr, TEXT("/AudioWidgets/AudioMaterialSlate/MI_AudioMaterialEnvelope_ADSR.MI_AudioMaterialEnvelope_ADSR"));
-
-	Set("AudioMaterialEnvelope.Style", FAudioMaterialEnvelopeStyle()
-		.SetMaterial(EnvelopeDefaultMaterial));
+	Set("AudioMaterialEnvelope.Style", FAudioMaterialEnvelopeStyle());
 
 	/**
 	**AudioMaterialButton Style
 	*/
-	UMaterialInterface* ButtonDefaultMaterial = LoadObject<UMaterialInterface>(nullptr, TEXT("/AudioWidgets/AudioMaterialSlate/MI_AudioMaterialButton.MI_AudioMaterialButton"));
-
-	Set("AudioMaterialButton.Style", FAudioMaterialButtonStyle()
-		.SetMaterial(ButtonDefaultMaterial));
+	Set("AudioMaterialButton.Style", FAudioMaterialButtonStyle());
 
 	/**
 	**AudioMaterialSlider Style
 	*/
-	UMaterialInterface* SliderDefaultMaterial = LoadObject<UMaterialInterface>(nullptr, TEXT("/AudioWidgets/AudioMaterialSlate/MI_AudioMaterialRoundedSlider.MI_AudioMaterialRoundedSlider"));
-
 	Set("AudioMaterialSlider.Style", FAudioMaterialSliderStyle()
-		.SetTextBoxStyle(FAudioTextBoxStyle::GetDefault())
-		.SetMaterial(SliderDefaultMaterial));
+		.SetTextBoxStyle(FAudioTextBoxStyle::GetDefault()));
 
 	FSlateStyleRegistry::RegisterSlateStyle(*this);
 }
