Option to use AudioMaterialButton in Metasound Editor bool node
#rnx

diff --git a/Engine/Plugins/Runtime/AudioWidgets/Source/AudioWidgets/Private/AudioMaterialSlate/AudioMaterialSlateTypes.cpp b/Engine/Plugins/Runtime/AudioWidgets/Source/AudioWidgets/Private/AudioMaterialSlate/AudioMaterialSlateTypes.cpp
--- a/Engine/Plugins/Runtime/AudioWidgets/Source/AudioWidgets/Private/AudioMaterialSlate/AudioMaterialSlateTypes.cpp
+++ b/Engine/Plugins/Runtime/AudioWidgets/Source/AudioWidgets/Private/AudioMaterialSlate/AudioMaterialSlateTypes.cpp
@@ -11,6 +11,23 @@ namespace AudioMaterialSlateTypesPrivate
 
 using namespace AudioMaterialSlateTypesPrivate;
 
+namespace AudioWidgets
+{
+	namespace SlateTypesPrivate
+	{
+		static const FLinearColor ButtonMainColor(0.098958f, 0.098958f, 0.098958f, 1.f);
+		static const FLinearColor ButtonAccentColor(0.341146f, 0.341146f, 0.341146f, 1.f);
+		static const FLinearColor ButtonPressedShadowColor(0.126558f, 0.138653f, 0.15f, 1.f);
+
+		static const FLinearColor BarMainColor(0.008f, 0.008f, 0.008f, 1.f);
+		static const FLinearColor BarAccentColor(0.005f, 0.005f, 0.005f, 1.f);
+		static const FLinearColor HandleMainColor(0.2f, 0.2f, 0.25f, 1.f);
+		static const FLinearColor HandleOutlineColor(0.15f, 0.15f, 0.15f, 1.f);
+	}
+}
+
+using namespace AudioWidgets;
+
 FAudioMaterialWidgetStyle::FAudioMaterialWidgetStyle()
 	:DesiredSize(32.f, 32.f)
 {
@@ -28,11 +45,11 @@ UMaterialInstanceDynamic* FAudioMaterialWidgetStyle::CreateDynamicMaterial(UObje
 }
 
 FAudioMaterialButtonStyle::FAudioMaterialButtonStyle()
-	: ButtonMainColor(FLinearColor::Gray)
+	: ButtonMainColor(SlateTypesPrivate::ButtonMainColor)
 	, ButtonShadowColor(FLinearColor::Black)
-	, ButtonAccentColor(FLinearColor::Gray)
+	, ButtonAccentColor(SlateTypesPrivate::ButtonAccentColor)
 	, ButtonPressedMainColor(FLinearColor::White)
-	, ButtonPressedShadowColor(FLinearColor::Gray)
+	, ButtonPressedShadowColor(SlateTypesPrivate::ButtonPressedShadowColor)
 	, ButtonPressedOutlineColor(FLinearColor::Blue)
 {
 	FString Path = PLUGIN_BASE_DIR + "MI_AudioMaterialButton.MI_AudioMaterialButton";
@@ -54,11 +71,11 @@ const FAudioMaterialButtonStyle& FAudioMaterialButtonStyle::GetDefault()
 }
 
 FAudioMaterialSliderStyle::FAudioMaterialSliderStyle()
-	: BarMainColor(FLinearColor(0.008f, 0.008f, 0.008f,1.f))
+	: BarMainColor(SlateTypesPrivate::BarMainColor)
 	, BarShadowColor(FLinearColor::Black)
-	, BarAccentColor(FLinearColor(0.005f, 0.005f, 0.005f, 1.f))
-	, HandleMainColor(FLinearColor(0.2f, 0.2f, 0.25f, 1.f))
-	, HandleOutlineColor(FLinearColor(0.15f, 0.15f, 0.15f, 1.f))
+	, BarAccentColor(SlateTypesPrivate::BarAccentColor)
+	, HandleMainColor(SlateTypesPrivate::HandleMainColor)
+	, HandleOutlineColor(SlateTypesPrivate::HandleOutlineColor)
 	, TextBoxStyle(FAudioTextBoxStyle::GetDefault())
 {
 	FString Path = PLUGIN_BASE_DIR + "MI_AudioMaterialRoundedSlider.MI_AudioMaterialRoundedSlider";
diff --git a/Engine/Plugins/Runtime/AudioWidgets/Source/AudioWidgets/Private/AudioMaterialSlate/SAudioMaterialButton.cpp b/Engine/Plugins/Runtime/AudioWidgets/Source/AudioWidgets/Private/AudioMaterialSlate/SAudioMaterialButton.cpp
--- a/Engine/Plugins/Runtime/AudioWidgets/Source/AudioWidgets/Private/AudioMaterialSlate/SAudioMaterialButton.cpp
+++ b/Engine/Plugins/Runtime/AudioWidgets/Source/AudioWidgets/Private/AudioMaterialSlate/SAudioMaterialButton.cpp
@@ -3,8 +3,8 @@
 
 #include "AudioMaterialSlate/SAudioMaterialButton.h"
 #include "AudioMaterialSlate/AudioMaterialSlateTypes.h"
-#include "SlateOptMacros.h"
 #include "Components/Widget.h"
+#include "SlateOptMacros.h"
 
 BEGIN_SLATE_FUNCTION_BUILD_OPTIMIZATION
 void SAudioMaterialButton::Construct(const FArguments& InArgs)
@@ -24,49 +24,72 @@ void SAudioMaterialButton::SetPressedState(bool InPressedState)
 	CommitNewState(InPressedState);
 }
 
-void SAudioMaterialButton::ApplyNewMaterial()
+UMaterialInstanceDynamic* SAudioMaterialButton::ApplyNewMaterial()
 {
 	if (AudioMaterialButtonStyle)
 	{
-		AudioMaterialButtonStyle->CreateDynamicMaterial(Owner.Get());
+		DynamicMaterial = AudioMaterialButtonStyle->CreateDynamicMaterial(Owner.Get());
 	}
+	return DynamicMaterial.Get();
+}
+
+void SAudioMaterialButton::SetDesiredSizeOverride(const FVector2D InSize)
+{
+	SetAttribute(DesiredSizeOverride, TAttribute<TOptional<FVector2D>>(InSize),EInvalidateWidgetReason::Layout);
 }
 
 int32 SAudioMaterialButton::OnPaint(const FPaintArgs& Args, const FGeometry& AllottedGeometry, const FSlateRect& MyCullingRect, FSlateWindowElementList& OutDrawElements, int32 LayerId, const FWidgetStyle& InWidgetStyle, bool bParentEnabled) const
 {
 	if (AudioMaterialButtonStyle)
 	{
-		UMaterialInstanceDynamic* DynamicMaterial = AudioMaterialButtonStyle->GetDynamicMaterial();
-
-		if (IsValid(DynamicMaterial))
+		if (DynamicMaterial.IsValid())
 		{
-			DynamicMaterial->SetVectorParameterValue(FName("MainColor"), AudioMaterialButtonStyle->ButtonMainColor);
-			DynamicMaterial->SetVectorParameterValue(FName("ShadowColor"), AudioMaterialButtonStyle->ButtonShadowColor);
-			DynamicMaterial->SetVectorParameterValue(FName("SmoothBevelColor"), AudioMaterialButtonStyle->ButtonAccentColor);
-			DynamicMaterial->SetVectorParameterValue(FName("Color_1"), AudioMaterialButtonStyle->ButtonPressedMainColor);
-			DynamicMaterial->SetVectorParameterValue(FName("Color_2"), AudioMaterialButtonStyle->ButtonPressedShadowColor);
-			DynamicMaterial->SetVectorParameterValue(FName("LedColor"), AudioMaterialButtonStyle->ButtonPressedOutlineColor);
-			DynamicMaterial->SetScalarParameterValue(FName("Click"), bIsPressedAttribute.Get());
-
-			DynamicMaterial->SetScalarParameterValue(FName("LocalWidth"), AllottedGeometry.GetLocalSize().X);
-			DynamicMaterial->SetScalarParameterValue(FName("LocalHeigth"), AllottedGeometry.GetLocalSize().Y);			
+			DynamicMaterial.Get()->SetVectorParameterValue(FName("MainColor"), AudioMaterialButtonStyle->ButtonMainColor);
+			DynamicMaterial.Get()->SetVectorParameterValue(FName("ShadowColor"), AudioMaterialButtonStyle->ButtonShadowColor);
+			DynamicMaterial.Get()->SetVectorParameterValue(FName("SmoothBevelColor"), AudioMaterialButtonStyle->ButtonAccentColor);
+			DynamicMaterial.Get()->SetVectorParameterValue(FName("Color_1"), AudioMaterialButtonStyle->ButtonPressedMainColor);
+			DynamicMaterial.Get()->SetVectorParameterValue(FName("Color_2"), AudioMaterialButtonStyle->ButtonPressedShadowColor);
+			DynamicMaterial.Get()->SetVectorParameterValue(FName("LedColor"), AudioMaterialButtonStyle->ButtonPressedOutlineColor);
+			DynamicMaterial.Get()->SetScalarParameterValue(FName("Click"), bIsPressedAttribute.Get());
+
+			DynamicMaterial.Get()->SetScalarParameterValue(FName("LocalWidth"), AllottedGeometry.GetLocalSize().X);
+			DynamicMaterial.Get()->SetScalarParameterValue(FName("LocalHeigth"), AllottedGeometry.GetLocalSize().Y);		
+			
+			const bool bEnabled = ShouldBeEnabled(bParentEnabled);
+			const ESlateDrawEffect DrawEffects = bEnabled ? ESlateDrawEffect::None : ESlateDrawEffect::DisabledEffect;
+
+			const FLinearColor FinalColorAndOpacity(InWidgetStyle.GetColorAndOpacityTint());
+
+			const float AllotedWidth = AllottedGeometry.GetLocalSize().X;
+			const float AllotedHeight = AllottedGeometry.GetLocalSize().Y;
+
+			const float ButtonRadius = FMath::Min(AllotedWidth, AllotedHeight) * 0.5f;
+			const FVector2D ButtonMidPoint(AllottedGeometry.GetLocalSize() * 0.5f);
+			const FVector2D ButtonDiameter(ButtonRadius * 2);
+
+			FSlateBrush Brush;
+			Brush.SetResourceObject(DynamicMaterial.Get());
+			FSlateDrawElement::MakeBox(OutDrawElements, LayerId++, AllottedGeometry.ToPaintGeometry(ButtonDiameter , FSlateLayoutTransform(ButtonMidPoint-ButtonRadius)),&Brush, DrawEffects, FinalColorAndOpacity);
 		}
-
-		const bool bEnabled = ShouldBeEnabled(bParentEnabled);
-		const ESlateDrawEffect DrawEffects = bEnabled ? ESlateDrawEffect::None : ESlateDrawEffect::DisabledEffect;
-
-		const FLinearColor FinalColorAndOpacity(InWidgetStyle.GetColorAndOpacityTint());
-
-		FSlateBrush Brush;
-		Brush.SetResourceObject(DynamicMaterial);
-		FSlateDrawElement::MakeBox(OutDrawElements, LayerId++, AllottedGeometry.ToPaintGeometry(),&Brush, DrawEffects, FinalColorAndOpacity);
-	}
+		else
+		{
+			if (AudioMaterialButtonStyle)
+			{
+				DynamicMaterial = AudioMaterialButtonStyle->CreateDynamicMaterial(Owner.Get());
+			}
+		}
+	}	
 
 	return LayerId;
 }
 
 FVector2D SAudioMaterialButton::ComputeDesiredSize(float) const
 {
+	if (DesiredSizeOverride.Get().IsSet())
+	{
+		return DesiredSizeOverride.Get().GetValue();
+	}
+
 	if (AudioMaterialButtonStyle)
 	{
 		return FVector2D(AudioMaterialButtonStyle->DesiredSize);
diff --git a/Engine/Plugins/Runtime/AudioWidgets/Source/AudioWidgets/Public/AudioMaterialSlate/SAudioMaterialButton.h b/Engine/Plugins/Runtime/AudioWidgets/Source/AudioWidgets/Public/AudioMaterialSlate/SAudioMaterialButton.h
--- a/Engine/Plugins/Runtime/AudioWidgets/Source/AudioWidgets/Public/AudioMaterialSlate/SAudioMaterialButton.h
+++ b/Engine/Plugins/Runtime/AudioWidgets/Source/AudioWidgets/Public/AudioMaterialSlate/SAudioMaterialButton.h
@@ -3,6 +3,7 @@
 #pragma once
 
 #include "AudioMaterialSlate/AudioMaterialSlateTypes.h"
+#include "AudioWidgetsStyle.h"
 #include "Framework/SlateDelegates.h"
 #include "Styling/ISlateStyle.h"
 #include "Styling/SlateWidgetStyleAsset.h"
@@ -19,6 +20,7 @@ class AUDIOWIDGETS_API SAudioMaterialButton : public SLeafWidget
 {
 public:
 	SLATE_BEGIN_ARGS(SAudioMaterialButton)
+	: _AudioMaterialButtonStyle(&FAudioWidgetsStyle::Get().GetWidgetStyle<FAudioMaterialButtonStyle>("AudioMaterialButton.Style"))
 	{}
 
 	/** The owner object*/
@@ -42,7 +44,10 @@ public:
 	void SetPressedState(bool InPressedState);
 
 	/** Apply new material to be used to render the Slate.*/
-	void ApplyNewMaterial();
+	UMaterialInstanceDynamic* ApplyNewMaterial();
+
+	/**Set desired size of the Slate*/
+	void SetDesiredSizeOverride(const FVector2D InSize);
 
 public:
 
@@ -66,10 +71,16 @@ private:
 	// Holds the owner of the Slate
 	TWeakObjectPtr<UObject> Owner;
 
+	// Holds the Modifiable Material that represent the Button
+	mutable TWeakObjectPtr<UMaterialInstanceDynamic> DynamicMaterial;
+
 	// Holds the style for the Slate
 	const FAudioMaterialButtonStyle* AudioMaterialButtonStyle = nullptr;
 
 	//Current pressed state of this button
 	TAttribute<bool> bIsPressedAttribute = false;
 
+	// Holds the optional desired size for the Slate
+	TAttribute<TOptional<FVector2D>> DesiredSizeOverride;
+
 };
diff --git a/Engine/Plugins/Runtime/Metasound/Source/MetasoundEditor/Private/MetasoundEditorGraphMemberDefaults.cpp b/Engine/Plugins/Runtime/Metasound/Source/MetasoundEditor/Private/MetasoundEditorGraphMemberDefaults.cpp
--- a/Engine/Plugins/Runtime/Metasound/Source/MetasoundEditor/Private/MetasoundEditorGraphMemberDefaults.cpp
+++ b/Engine/Plugins/Runtime/Metasound/Source/MetasoundEditor/Private/MetasoundEditorGraphMemberDefaults.cpp
@@ -69,6 +69,16 @@ void UMetasoundEditorGraphMemberDefaultBool::UpdatePreviewInstance(const Metasou
 	InParameterInterface->SetBoolParameter(InParameterName, Default.Value);
 }
 
+void UMetasoundEditorGraphMemberDefaultBool::SetDefault(const bool InDefault)
+{
+	Default.Value = InDefault;
+}
+
+const bool UMetasoundEditorGraphMemberDefaultBool::GetDefault()
+{
+	return Default.Value;
+}
+
 FMetasoundFrontendLiteral UMetasoundEditorGraphMemberDefaultBoolArray::GetDefault() const
 {
 	TArray<bool> BoolArray;
diff --git a/Engine/Plugins/Runtime/Metasound/Source/MetasoundEditor/Private/MetasoundEditorModule.cpp b/Engine/Plugins/Runtime/Metasound/Source/MetasoundEditor/Private/MetasoundEditorModule.cpp
--- a/Engine/Plugins/Runtime/Metasound/Source/MetasoundEditor/Private/MetasoundEditorModule.cpp
+++ b/Engine/Plugins/Runtime/Metasound/Source/MetasoundEditor/Private/MetasoundEditorModule.cpp
@@ -839,6 +839,7 @@ namespace Metasound
 
 				LiteralCustomizationFactories.Add(UMetasoundEditorGraphMemberDefaultLiteral::StaticClass(), MakeUnique<FMetasoundDefaultLiteralCustomizationFactory>());
 				LiteralCustomizationFactories.Add(UMetasoundEditorGraphMemberDefaultFloat::StaticClass(), MakeUnique<FMetasoundFloatLiteralCustomizationFactory>());
+				LiteralCustomizationFactories.Add(UMetasoundEditorGraphMemberDefaultBool::StaticClass(), MakeUnique<FMetasoundBoolLiteralCustomizationFactory>());
 				LiteralCustomizationFactories.Add(UMetasoundEditorGraphMemberDefaultObjectArray::StaticClass(), MakeUnique<FMetasoundObjectArrayLiteralCustomizationFactory>());
 
 				StyleSet = MakeShared<FSlateStyle>();
diff --git a/Engine/Plugins/Runtime/Metasound/Source/MetasoundEditor/Private/MetasoundEditorSettings.h b/Engine/Plugins/Runtime/Metasound/Source/MetasoundEditor/Private/MetasoundEditorSettings.h
--- a/Engine/Plugins/Runtime/Metasound/Source/MetasoundEditor/Private/MetasoundEditorSettings.h
+++ b/Engine/Plugins/Runtime/Metasound/Source/MetasoundEditor/Private/MetasoundEditorSettings.h
@@ -27,6 +27,13 @@ enum class EMetasoundMemberDefaultWidget : uint8
 	RadialSlider UMETA(DisplayName = "Knob"),
 };
 
+UENUM()
+enum class EMetasoundBoolMemberDefaultWidget : uint8
+{
+	None,
+	Button
+};
+
 UENUM()
 enum class EMetasoundActiveDetailView : uint8
 {
diff --git a/Engine/Plugins/Runtime/Metasound/Source/MetasoundEditor/Private/MetasoundNodeDetailCustomization.cpp b/Engine/Plugins/Runtime/Metasound/Source/MetasoundEditor/Private/MetasoundNodeDetailCustomization.cpp
--- a/Engine/Plugins/Runtime/Metasound/Source/MetasoundEditor/Private/MetasoundNodeDetailCustomization.cpp
+++ b/Engine/Plugins/Runtime/Metasound/Source/MetasoundEditor/Private/MetasoundNodeDetailCustomization.cpp
@@ -361,6 +361,59 @@ namespace Metasound
 			return DefaultRows;
 		}
 
+		FMetasoundBoolLiteralCustomization::~FMetasoundBoolLiteralCustomization()
+		{
+		}
+
+		TArray<IDetailPropertyRow*> FMetasoundBoolLiteralCustomization::CustomizeLiteral(UMetasoundEditorGraphMemberDefaultLiteral& InLiteral, IDetailLayoutBuilder& InDetailLayout)
+		{
+			check(DefaultCategoryBuilder);
+
+			UMetasoundEditorGraphMemberDefaultBool* DefaultBool = Cast<UMetasoundEditorGraphMemberDefaultBool>(&InLiteral);
+			if (!ensure(DefaultBool))
+			{
+				return { };
+			}
+			BoolLiteral = DefaultBool;
+
+			TArray<IDetailPropertyRow*> DefaultRows;
+			TSharedPtr<IPropertyHandle> DefaultValueHandle;
+			IDetailPropertyRow* Row = DefaultCategoryBuilder->AddExternalObjectProperty(TArray<UObject*>({ DefaultBool }), GET_MEMBER_NAME_CHECKED(UMetasoundEditorGraphMemberDefaultBool, Default));
+			if (ensure(Row))
+			{
+				DefaultRows.Add(Row);
+				DefaultValueHandle = Row->GetPropertyHandle();
+			}
+
+			// Enable widget options for editable inputs only 
+			bool bShowWidgetOptions = false;
+			if (const UMetasoundEditorGraphInput* ParentMember = Cast <UMetasoundEditorGraphInput>(InLiteral.GetParentMember()))
+			{								
+				if (const UMetasoundEditorGraph* OwningGraph = ParentMember->GetOwningGraph())
+				{
+					// Disable widget options for constructor inputs for now to prevent changing default value via widget while playing 
+					bShowWidgetOptions = OwningGraph->IsEditable() && ParentMember->GetVertexAccessType() == EMetasoundFrontendVertexAccessType::Reference;
+				}
+			}
+
+			// add input widget properties
+			if (bShowWidgetOptions)
+			{
+				Frontend::FDataTypeRegistryInfo DataTypeInfo;
+				MemberCustomizationPrivate::GetDataTypeFromElementPropertyHandle(DefaultValueHandle, DataTypeInfo);
+
+				//TODO: Check if using AudioMaterialSlates, only then add.
+				//Only display widgetType option if this is not trigger
+				if (MemberCustomizationPrivate::GetPrimitiveTypeName(DataTypeInfo) != Metasound::GetMetasoundDataTypeName<Metasound::FTrigger>())
+				{
+					IDetailCategoryBuilder& WidgetCategoryBuilder = InDetailLayout.EditCategory("EditorOptions");
+					DefaultRows.Add(WidgetCategoryBuilder.AddExternalObjectProperty(TArray<UObject*>({ DefaultBool }), GET_MEMBER_NAME_CHECKED(UMetasoundEditorGraphMemberDefaultBool, WidgetType)));
+				}
+			}
+
+			return DefaultRows;
+		}
+
 		TArray<IDetailPropertyRow*> FMetasoundObjectArrayLiteralCustomization::CustomizeLiteral(UMetasoundEditorGraphMemberDefaultLiteral& InLiteral, IDetailLayoutBuilder& InDetailLayout)
 		{
 			check(DefaultCategoryBuilder);
diff --git a/Engine/Plugins/Runtime/Metasound/Source/MetasoundEditor/Private/MetasoundNodeDetailCustomization.h b/Engine/Plugins/Runtime/Metasound/Source/MetasoundEditor/Private/MetasoundNodeDetailCustomization.h
--- a/Engine/Plugins/Runtime/Metasound/Source/MetasoundEditor/Private/MetasoundNodeDetailCustomization.h
+++ b/Engine/Plugins/Runtime/Metasound/Source/MetasoundEditor/Private/MetasoundNodeDetailCustomization.h
@@ -188,6 +188,20 @@ namespace Metasound
 			virtual TArray<IDetailPropertyRow*> CustomizeLiteral(UMetasoundEditorGraphMemberDefaultLiteral& InLiteral, IDetailLayoutBuilder& InDetailLayout) override;
 		};
 
+		class FMetasoundBoolLiteralCustomization : public FMetasoundDefaultLiteralCustomizationBase
+		{
+			TWeakObjectPtr<UMetasoundEditorGraphMemberDefaultBool> BoolLiteral;
+
+		public:
+			FMetasoundBoolLiteralCustomization(IDetailCategoryBuilder& InDefaultCategoryBuilder)
+				: FMetasoundDefaultLiteralCustomizationBase(InDefaultCategoryBuilder)
+			{
+			}
+			virtual ~FMetasoundBoolLiteralCustomization();
+
+			virtual TArray<IDetailPropertyRow*> CustomizeLiteral(UMetasoundEditorGraphMemberDefaultLiteral& InLiteral, IDetailLayoutBuilder& InDetailLayout) override;
+		};
+
 		// Customization to support drag-and-drop of Proxy UObject types on underlying members that are structs.
 		// Struct ownership of objects required to customize asset filters based on dynamic UObject MetaSound Registry DataTypes.
 		class FMetasoundObjectArrayLiteralCustomization : public FMetasoundDefaultLiteralCustomizationBase
@@ -221,6 +235,16 @@ namespace Metasound
 				return TUniquePtr<FMetasoundDefaultLiteralCustomizationBase>(new FMetasoundFloatLiteralCustomization(DefaultCategoryBuilder));
 			}
 		};
+		
+		// Customization to support bool widgets (ex. Buttons)
+		class FMetasoundBoolLiteralCustomizationFactory : public IMemberDefaultLiteralCustomizationFactory
+		{
+		public:
+			virtual TUniquePtr<FMetasoundDefaultLiteralCustomizationBase> CreateLiteralCustomization(IDetailCategoryBuilder& DefaultCategoryBuilder) const override
+			{
+				return TUniquePtr<FMetasoundDefaultLiteralCustomizationBase>(new FMetasoundBoolLiteralCustomization(DefaultCategoryBuilder));
+			}
+		};
 
 		class FMetasoundObjectArrayLiteralCustomizationFactory : public IMemberDefaultLiteralCustomizationFactory
 		{
diff --git a/Engine/Plugins/Runtime/Metasound/Source/MetasoundEditor/Private/SMetasoundGraphNode.cpp b/Engine/Plugins/Runtime/Metasound/Source/MetasoundEditor/Private/SMetasoundGraphNode.cpp
--- a/Engine/Plugins/Runtime/Metasound/Source/MetasoundEditor/Private/SMetasoundGraphNode.cpp
+++ b/Engine/Plugins/Runtime/Metasound/Source/MetasoundEditor/Private/SMetasoundGraphNode.cpp
@@ -1,6 +1,7 @@
 // Copyright Epic Games, Inc. All Rights Reserved.
 #include "SMetasoundGraphNode.h"
 
+#include "AudioMaterialSlate/SAudioMaterialButton.h"
 #include "AudioMaterialSlate/SAudioMaterialLabeledKnob.h"
 #include "AudioMaterialSlate/SAudioMaterialLabeledSlider.h"
 #include "AudioParameterControllerInterface.h"
@@ -669,7 +670,7 @@ namespace Metasound
 			TSharedPtr<SWidget> OuterContentBox; // currently only used for input float nodes to accommodate the input widget
 
 			// If editable float input node and not constructor input, check if custom widget required
-			bool bShowInputWidget = false;
+			bool bShowContentWidget = false;
 			if (UMetasoundEditorGraphInput* GraphMember = Cast<UMetasoundEditorGraphInput>(GetMetaSoundMember()))
 			{
 				const UMetasoundEditorGraph* OwningGraph = GraphMember->GetOwningGraph();
@@ -682,7 +683,7 @@ namespace Metasound
 						static const FVector2D SliderDesiredSizeVertical = FVector2D(30.0f, 250.0f);
 						static const FVector2D RadialSliderDesiredSize = FVector2D(56.0f, 87.0f);
 
-						bShowInputWidget = true;
+						bShowContentWidget = true;
 
 						auto OnValueChangedLambda = [DefaultFloat, GraphMember, this](float Value)
 						{
@@ -960,6 +961,56 @@ namespace Metasound
 							}
 						});
 					}
+					else if (UMetasoundEditorGraphMemberDefaultBool* DefaultBool = Cast<UMetasoundEditorGraphMemberDefaultBool>(GraphMember->GetLiteral()); Metasound::Editor::GraphNodePrivate::UseAudioMaterialWidgets)
+					{
+						bool bIsNotTriggerNode = GraphMember->GetDataType() != Metasound::GetMetasoundDataTypeName<Metasound::FTrigger>();
+
+						if ((bIsNotTriggerNode && IsValid(DefaultBool)) && DefaultBool->WidgetType != EMetasoundBoolMemberDefaultWidget::None)
+						{
+							constexpr float WidgetPadding = 3.0f;
+							static const FVector2D ButtonDesiredSize = FVector2D(56.0f, 87.0f);
+
+							auto OnboolValueChangedLambda = [DefaultBool, GraphMember, this](bool Value)
+								{
+									if (MaterialButtonWidget.IsValid())
+									{
+										if (!bIsInputWidgetTransacting)
+										{
+											GEditor->BeginTransaction(LOCTEXT("MetasoundGraphNode_MetasoundSetInputDefault", "Set MetaSound Input Default"));
+											bIsInputWidgetTransacting = true;
+										}
+										GraphMember->GetOwningGraph()->GetMetasound()->Modify();
+										DefaultBool->Modify();
+
+										constexpr bool bPostTransaction = true;
+										DefaultBool->SetDefault(Value);
+										GraphMember->UpdateFrontendDefaultLiteral(bPostTransaction);
+									}
+								};
+
+							bShowContentWidget = true;
+							SAssignNew(OuterContentBox, SVerticalBox)
+								+ SVerticalBox::Slot()
+								.HAlign(HAlign_Right)
+								.VAlign(VAlign_Center)
+								.AutoHeight()
+								[
+									ContentBox.ToSharedRef()
+								]
+								+ SVerticalBox::Slot()
+								.HAlign(HAlign_Fill)
+								.VAlign(VAlign_Top)
+								.Padding(WidgetPadding, 0.0f, WidgetPadding, WidgetPadding)
+								.AutoHeight()
+								[
+									SAssignNew(MaterialButtonWidget, SAudioMaterialButton)
+										.OnBooleanValueChanged_Lambda(OnboolValueChangedLambda)									
+								];
+
+							MaterialButtonWidget->SetPressedState(DefaultBool->GetDefault());
+							MaterialButtonWidget->SetDesiredSizeOverride(ButtonDesiredSize);
+						}
+					}
 				}
 			}
 	
@@ -1007,7 +1058,7 @@ namespace Metasound
 				.VAlign(VAlign_Fill)
 				.Padding(FMargin(0,3))
 				[
-					(bShowInputWidget ? OuterContentBox : ContentBox).ToSharedRef()
+					(bShowContentWidget ? OuterContentBox : ContentBox).ToSharedRef()
 				];
 		}
 
diff --git a/Engine/Plugins/Runtime/Metasound/Source/MetasoundEditor/Private/SMetasoundGraphNode.h b/Engine/Plugins/Runtime/Metasound/Source/MetasoundEditor/Private/SMetasoundGraphNode.h
--- a/Engine/Plugins/Runtime/Metasound/Source/MetasoundEditor/Private/SMetasoundGraphNode.h
+++ b/Engine/Plugins/Runtime/Metasound/Source/MetasoundEditor/Private/SMetasoundGraphNode.h
@@ -14,6 +14,7 @@
 
 // Forward Declarations
 class SAudioInputWidget;
+class SAudioMaterialButton;
 class SGraphPin;
 class SVerticalBox;
 class UMetasoundEditorGraphMember;
@@ -79,6 +80,8 @@ namespace Metasound
 
 			// Slider widget for float input
 			TSharedPtr<SAudioInputWidget> InputWidget;
+			//Button Widget for bool input.
+			TSharedPtr <SAudioMaterialButton> MaterialButtonWidget;
 			// Handle for on value changed delegate for input slider 
 			FDelegateHandle InputSliderOnValueChangedDelegateHandle;
 			// Handle for on input slider range changed  
diff --git a/Engine/Plugins/Runtime/Metasound/Source/MetasoundEditor/Public/MetasoundEditorGraphMemberDefaults.h b/Engine/Plugins/Runtime/Metasound/Source/MetasoundEditor/Public/MetasoundEditorGraphMemberDefaults.h
--- a/Engine/Plugins/Runtime/Metasound/Source/MetasoundEditor/Public/MetasoundEditorGraphMemberDefaults.h
+++ b/Engine/Plugins/Runtime/Metasound/Source/MetasoundEditor/Public/MetasoundEditorGraphMemberDefaults.h
@@ -50,10 +50,16 @@ public:
 	UPROPERTY(EditAnywhere, Category = DefaultValue)
 	FMetasoundEditorGraphMemberDefaultBoolRef Default;
 
+	UPROPERTY(EditAnywhere, Category = Widget, meta = (DisplayName = "Widget"))
+	EMetasoundBoolMemberDefaultWidget WidgetType = EMetasoundBoolMemberDefaultWidget::None;
+
 	virtual FMetasoundFrontendLiteral GetDefault() const override;
 	virtual EMetasoundFrontendLiteralType GetLiteralType() const override;
 	virtual void SetFromLiteral(const FMetasoundFrontendLiteral& InLiteral) override;
 	virtual void UpdatePreviewInstance(const Metasound::FVertexName& InParameterName, TScriptInterface<IAudioParameterControllerInterface>& InParameterInterface) const override;
+
+	void SetDefault(const bool InDefault);
+	const bool GetDefault();
 };
 
 UCLASS(MinimalAPI)
