Removed the Envelope style to be an option to create Slate Widget Styles
#rnx

diff --git a/Engine/Plugins/Runtime/AudioWidgets/Source/AudioWidgets/Public/AudioMaterialSlate/AudioMaterialStyleContainers.h b/Engine/Plugins/Runtime/AudioWidgets/Source/AudioWidgets/Public/AudioMaterialSlate/AudioMaterialStyleContainers.h
--- a/Engine/Plugins/Runtime/AudioWidgets/Source/AudioWidgets/Public/AudioMaterialSlate/AudioMaterialStyleContainers.h
+++ b/Engine/Plugins/Runtime/AudioWidgets/Source/AudioWidgets/Public/AudioMaterialSlate/AudioMaterialStyleContainers.h
@@ -74,20 +74,3 @@ public:
 		return static_cast<const struct FSlateWidgetStyle*>(&SliderStyle);
 	}
 };
-
-UCLASS(hidecategories = Object, MinimalAPI)
-class UAudioMaterialEnvelopeWidgetStyle : public USlateWidgetStyleContainerBase
-{
-	GENERATED_BODY()
-
-public:
-
-	/** The actual data describing the AudioMaterialSlider appearance. */
-	UPROPERTY(Category = Appearance, EditAnywhere, meta = (ShowOnlyInnerProperties))
-	FAudioMaterialEnvelopeStyle EnvelopeStyle;
-
-	virtual const struct FSlateWidgetStyle* const GetStyle() const override
-	{
-		return static_cast<const struct FSlateWidgetStyle*>(&EnvelopeStyle);
-	}
-};
