Spectrum analyzer for white noise of 0dB RMS to display as 0dB.

diff --git a/Engine/Plugins/Runtime/AudioWidgets/Source/AudioWidgets/Private/AudioSpectrumAnalyzer.cpp b/Engine/Plugins/Runtime/AudioWidgets/Source/AudioWidgets/Private/AudioSpectrumAnalyzer.cpp
--- a/Engine/Plugins/Runtime/AudioWidgets/Source/AudioWidgets/Private/AudioSpectrumAnalyzer.cpp
+++ b/Engine/Plugins/Runtime/AudioWidgets/Source/AudioWidgets/Private/AudioSpectrumAnalyzer.cpp
@@ -4,13 +4,45 @@
 
 #include "AudioDevice.h"
 #include "AudioDeviceManager.h"
+#include "ConstantQFactory.h"
 #include "DSP/EnvelopeFollower.h"
 #include "Framework/MultiBox/MultiBoxBuilder.h"
+#include "SynesthesiaSpectrumAnalysisFactory.h"
 
 #define LOCTEXT_NAMESPACE "FAudioSpectrumAnalyzer"
 
 namespace AudioWidgets
 {
+	namespace AudioSpectrumAnalyzerPrivate
+	{
+		inline float GetWindowCompensationPowerGain(Audio::EWindowType WindowType, int32 FFTSize)
+		{
+			ensure(FFTSize <= 16384);
+			ensure(FMath::IsPowerOfTwo(FFTSize));
+
+			using namespace Audio;
+
+			// Create a temporary buffer and initialize it to +1 DC:
+			FAlignedFloatBuffer Samples;
+			Samples.SetNumUninitialized(FFTSize);
+			Audio::ArraySetToConstantInplace(Samples, 1.0f);
+
+			// Initialize the window in the same manner as FSynesthesiaSpectrumAnalyzer and FConstantQAnalyzer:
+			FWindow Window(WindowType, FFTSize, /*NumChannels=*/1, /*bIsPeriodic=*/false);
+
+			// Apply window to DC signal:
+			Window.ApplyToBuffer(Samples.GetData());
+
+			// Calculate the mean square of the windowed signal:
+			float WindowedDCMeanSquare = 1.0f;
+			Audio::ArrayMeanSquared(Samples, WindowedDCMeanSquare);
+
+			// Return the power gain required to reverse the effect of the windowing process on the RMS of DC:
+			constexpr float DCMeanSquare = 1.0f;
+			return DCMeanSquare / WindowedDCMeanSquare;
+		}
+	}
+
 	FAudioSpectrumAnalyzer::FAudioSpectrumAnalyzer(int32 InNumChannels, Audio::FDeviceId InAudioDeviceId, TObjectPtr<UAudioBus> InExternalAudioBus)
 		: SpectrumAnalysisSettings(NewObject<USynesthesiaSpectrumAnalysisSettings>())
 		, ConstantQSettings(NewObject<UConstantQSettings>())
@@ -136,6 +168,14 @@ namespace AudioWidgets
 
 					// Init spectrum data:
 					ARSmoothedSquaredMagnitudes = SpectrumResults.SpectrumValues;
+
+					// Update the window compensation power gain:
+					TUniquePtr<Audio::IAnalyzerSettings> Settings = SpectrumAnalysisSettings->GetSettings(SampleRate, 1);
+					const Audio::FSynesthesiaSpectrumAnalysisSettings* ConcreteSettings = static_cast<Audio::FSynesthesiaSpectrumAnalysisSettings*>(Settings.Get());
+					WindowCompensationPowerGain = AudioSpectrumAnalyzerPrivate::GetWindowCompensationPowerGain(ConcreteSettings->WindowType, ConcreteSettings->FFTSize);
+
+					// Apply window compensation power gain:
+					Audio::ArrayMultiplyByConstantInPlace(ARSmoothedSquaredMagnitudes, WindowCompensationPowerGain);
 				}
 
 				PrevTimeStamp = SpectrumResults.TimeSeconds;
@@ -161,6 +201,14 @@ namespace AudioWidgets
 
 					// Init spectrum data:
 					ARSmoothedSquaredMagnitudes = SpectrumResults.SpectrumValues;
+
+					// Update the window compensation power gain:
+					TUniquePtr<Audio::IAnalyzerSettings> Settings = ConstantQSettings->GetSettings(0.0f, 1);
+					const Audio::FConstantQSettings* ConcreteSettings = static_cast<Audio::FConstantQSettings*>(Settings.Get());
+					WindowCompensationPowerGain = AudioSpectrumAnalyzerPrivate::GetWindowCompensationPowerGain(ConcreteSettings->WindowType, ConcreteSettings->FFTSize);
+
+					// Apply window compensation power gain:
+					Audio::ArrayMultiplyByConstantInPlace(ARSmoothedSquaredMagnitudes, WindowCompensationPowerGain);
 				}
 
 				PrevTimeStamp = SpectrumResults.TimeSeconds;
@@ -179,7 +227,7 @@ namespace AudioWidgets
 		for (int Index = 0; Index < SquaredMagnitudes.Num(); Index++)
 		{
 			const float OldValue = ARSmoothedSquaredMagnitudes[Index];
-			const float NewValue = SquaredMagnitudes[Index];
+			const float NewValue = WindowCompensationPowerGain * SquaredMagnitudes[Index];
 			const float ARSmootherCoefficient = (NewValue >= OldValue) ? AttackRelease.GetAttackTimeSamples() : AttackRelease.GetReleaseTimeSamples();
 			ARSmoothedSquaredMagnitudes[Index] = FMath::Lerp(NewValue, OldValue, ARSmootherCoefficient);
 		}
diff --git a/Engine/Plugins/Runtime/AudioWidgets/Source/AudioWidgets/Public/AudioSpectrumAnalyzer.h b/Engine/Plugins/Runtime/AudioWidgets/Source/AudioWidgets/Public/AudioSpectrumAnalyzer.h
--- a/Engine/Plugins/Runtime/AudioWidgets/Source/AudioWidgets/Public/AudioSpectrumAnalyzer.h
+++ b/Engine/Plugins/Runtime/AudioWidgets/Source/AudioWidgets/Public/AudioSpectrumAnalyzer.h
@@ -94,6 +94,7 @@ namespace AudioWidgets
 		EAudioSpectrumAnalyzerType AnalyzerType = EAudioSpectrumAnalyzerType::CQT;
 
 		TOptional<float> PrevTimeStamp;
+		float WindowCompensationPowerGain = 1.0f;
 		float AttackTimeMsec = 300.0f;
 		float ReleaseTimeMsec = 300.0f;
 		bool bIsAnalogAttackRelease = false;
