Spectrum analyzer tilt.

diff --git a/Engine/Plugins/Runtime/AudioWidgets/Source/AudioWidgets/Private/SAudioSpectrumPlot.cpp b/Engine/Plugins/Runtime/AudioWidgets/Source/AudioWidgets/Private/SAudioSpectrumPlot.cpp
--- a/Engine/Plugins/Runtime/AudioWidgets/Source/AudioWidgets/Private/SAudioSpectrumPlot.cpp
+++ b/Engine/Plugins/Runtime/AudioWidgets/Source/AudioWidgets/Private/SAudioSpectrumPlot.cpp
@@ -203,6 +203,8 @@ void SAudioSpectrumPlot::Construct(const FArguments& InArgs)
 	ViewMaxFrequency = InArgs._ViewMaxFrequency;
 	ViewMinSoundLevel = InArgs._ViewMinSoundLevel;
 	ViewMaxSoundLevel = InArgs._ViewMaxSoundLevel;
+	TiltExponent = InArgs._TiltExponent;
+	TiltPivotFrequency = InArgs._TiltPivotFrequency;
 	bDisplayFrequencyAxisLabels = InArgs._DisplayFrequencyAxisLabels;
 	bDisplaySoundLevelAxisLabels = InArgs._DisplaySoundLevelAxisLabels;
 	bDisplayFrequencyGridLines = InArgs._DisplayFrequencyGridLines;
@@ -442,12 +444,15 @@ int32 SAudioSpectrumPlot::DrawPowerSpectrum(const FGeometry& AllottedGeometry, F
 		TArray<FVector2f> DataPoints;
 		DataPoints.Reserve(NumFrequencies);
 
+		const float TiltExponentValue = TiltExponent.Get();
+		const float TiltPivotFrequencyValue = TiltPivotFrequency.Get();
 		const float ClampMinFrequency = (FrequencyAxisScale.Get() == EAudioSpectrumPlotFrequencyAxisScale::Logarithmic) ? 0.00001f : -FLT_MAX; // Cannot plot DC with log scale.
 		const float ClampMinMagnitudeSquared = FMath::Pow(10.0f, -200.0f / 10.0f); // Clamp at -200dB
 		for (int Index = 0; Index < NumFrequencies; Index++)
 		{
 			const float Frequency = FMath::Max(PowerSpectrum.CenterFrequencies[Index], ClampMinFrequency);
-			const float MagnitudeSquared = FMath::Max(PowerSpectrum.SquaredMagnitudes[Index], ClampMinMagnitudeSquared);
+			const float TiltPowerGain = FMath::Pow(Frequency / TiltPivotFrequencyValue, TiltExponentValue);
+			const float MagnitudeSquared = FMath::Max(TiltPowerGain * PowerSpectrum.SquaredMagnitudes[Index], ClampMinMagnitudeSquared);
 			const float SoundLevel = 10.0f * FMath::LogX(10.0f, MagnitudeSquared);
 			DataPoints.Add({ Frequency, SoundLevel });
 		}
@@ -600,6 +605,24 @@ FLinearColor SAudioSpectrumPlot::GetSpectrumColor(const FWidgetStyle& InWidgetSt
 	return SlateColor.GetColor(InWidgetStyle) * InWidgetStyle.GetColorAndOpacityTint();
 }
 
+float SAudioSpectrumPlot::GetTiltExponentValue(const EAudioSpectrumPlotTilt InTilt)
+{
+	switch (InTilt)
+	{
+	default:
+	case EAudioSpectrumPlotTilt::NoTilt:
+		return 0.0f;
+	case EAudioSpectrumPlotTilt::Plus1_5dBPerOctave:
+		return 0.5f;
+	case EAudioSpectrumPlotTilt::Plus3dBPerOctave:
+		return 1.0f;
+	case EAudioSpectrumPlotTilt::Plus4_5dBPerOctave:
+		return 1.5f;
+	case EAudioSpectrumPlotTilt::Plus6dBPerOctave:
+		return 2.0f;
+	}
+}
+
 TSharedRef<SWidget> SAudioSpectrumPlot::BuildDefaultContextMenu()
 {
 	constexpr bool bShouldCloseWindowAfterMenuSelection = true;
@@ -607,6 +630,14 @@ TSharedRef<SWidget> SAudioSpectrumPlot::BuildDefaultContextMenu()
 
 	MenuBuilder.BeginSection(ContextMenuExtensionHook, LOCTEXT("DisplayOptions", "Display Options"));
 
+	if (!TiltExponent.IsBound())
+	{
+		MenuBuilder.AddSubMenu(
+			LOCTEXT("TiltSpectrum", "Tilt Spectrum"),
+			FText(),
+			FNewMenuDelegate::CreateSP(this, &SAudioSpectrumPlot::BuildTiltSpectrumSubMenu));
+	}
+
 	if (!FrequencyAxisPixelBucketMode.IsBound())
 	{
 		MenuBuilder.AddSubMenu(
@@ -658,6 +689,33 @@ TSharedRef<SWidget> SAudioSpectrumPlot::BuildDefaultContextMenu()
 	return MenuBuilder.MakeWidget();
 }
 
+void SAudioSpectrumPlot::BuildTiltSpectrumSubMenu(FMenuBuilder& SubMenu)
+{
+	const UEnum* EnumClass = StaticEnum<EAudioSpectrumPlotTilt>();
+	const int32 NumEnumValues = EnumClass->NumEnums() - 1; // Exclude 'MAX' enum value.
+	for (int32 Index = 0; Index < NumEnumValues; Index++)
+	{
+		const auto EnumValue = static_cast<EAudioSpectrumPlotTilt>(EnumClass->GetValueByIndex(Index));
+		const float TiltExponentValue = GetTiltExponentValue(EnumValue);
+
+		SubMenu.AddMenuEntry(
+			EnumClass->GetDisplayNameTextByIndex(Index),
+#if WITH_EDITOR
+			EnumClass->GetToolTipTextByIndex(Index),
+#else
+			FText(),
+#endif
+			FSlateIcon(),
+			FUIAction(
+				FExecuteAction::CreateSPLambda(this, [this, TiltExponentValue]() { TiltExponent = TiltExponentValue; }),
+				FCanExecuteAction(),
+				FIsActionChecked::CreateSPLambda(this, [this, TiltExponentValue]() { return (TiltExponent.Get() == TiltExponentValue); })
+			),
+			NAME_None,
+			EUserInterfaceActionType::ToggleButton);
+	}
+}
+
 void SAudioSpectrumPlot::BuildFrequencyAxisScaleSubMenu(FMenuBuilder& SubMenu)
 {
 	const UEnum* EnumClass = StaticEnum<EAudioSpectrumPlotFrequencyAxisScale>();
diff --git a/Engine/Plugins/Runtime/AudioWidgets/Source/AudioWidgets/Public/SAudioSpectrumPlot.h b/Engine/Plugins/Runtime/AudioWidgets/Source/AudioWidgets/Public/SAudioSpectrumPlot.h
--- a/Engine/Plugins/Runtime/AudioWidgets/Source/AudioWidgets/Public/SAudioSpectrumPlot.h
+++ b/Engine/Plugins/Runtime/AudioWidgets/Source/AudioWidgets/Public/SAudioSpectrumPlot.h
@@ -9,6 +9,16 @@
 #include "Styling/SlateWidgetStyleAsset.h"
 #include "Widgets/SCompoundWidget.h"
 
+UENUM(BlueprintType)
+enum class EAudioSpectrumPlotTilt : uint8
+{
+	NoTilt UMETA(ToolTip = "0 dB/octave slope (white noise is flat)."),
+	Plus1_5dBPerOctave UMETA(DisplayName = "1.5 dB/octave", ToolTip = "1.5 dB/octave slope."),
+	Plus3dBPerOctave UMETA(DisplayName = "3 dB/octave", ToolTip = "3 dB/octave slope (pink noise is flat)."),
+	Plus4_5dBPerOctave UMETA(DisplayName = "4.5 dB/octave", ToolTip = "4.5 dB/octave slope."),
+	Plus6dBPerOctave UMETA(DisplayName = "6 dB/octave", ToolTip = "6 dB/octave slope (Brownian noise is flat)."),
+};
+
 UENUM(BlueprintType)
 enum class EAudioSpectrumPlotFrequencyAxisScale : uint8
 {
@@ -119,6 +129,8 @@ public:
 		, _ViewMaxFrequency(20000.0f)
 		, _ViewMinSoundLevel(-60.0f)
 		, _ViewMaxSoundLevel(12.0f)
+		, _TiltExponent(0.0f)
+		, _TiltPivotFrequency(24000.0f)
 		, _DisplayFrequencyAxisLabels(true)
 		, _DisplaySoundLevelAxisLabels(true)
 		, _DisplayFrequencyGridLines(true)
@@ -136,6 +148,8 @@ public:
 		SLATE_ATTRIBUTE(float, ViewMaxFrequency)
 		SLATE_ATTRIBUTE(float, ViewMinSoundLevel)
 		SLATE_ATTRIBUTE(float, ViewMaxSoundLevel)
+		SLATE_ATTRIBUTE(float, TiltExponent)
+		SLATE_ATTRIBUTE(float, TiltPivotFrequency)
 		SLATE_ATTRIBUTE(bool, DisplayFrequencyAxisLabels)
 		SLATE_ATTRIBUTE(bool, DisplaySoundLevelAxisLabels)
 		SLATE_ATTRIBUTE(bool, DisplayFrequencyGridLines)
@@ -158,6 +172,8 @@ public:
 	void SetViewMaxFrequency(float InViewMaxFrequency) { ViewMaxFrequency = InViewMaxFrequency; }
 	void SetViewMinSoundLevel(float InViewMinSoundLevel) { ViewMinSoundLevel = InViewMinSoundLevel; }
 	void SetViewMaxSoundLevel(float InViewMaxSoundLevel) { ViewMaxSoundLevel = InViewMaxSoundLevel; }
+	void SetTiltExponent(float InTiltExponent) { TiltExponent = InTiltExponent; }
+	void SetTiltPivotFrequency(float InTiltPivotFrequency) { TiltPivotFrequency = InTiltPivotFrequency; }
 	void SetDisplayFrequencyAxisLabels(bool bInDisplayFrequencyAxisLabels) { bDisplayFrequencyAxisLabels = bInDisplayFrequencyAxisLabels; }
 	void SetDisplaySoundLevelAxisLabels(bool bInDisplaySoundLevelAxisLabels) { bDisplaySoundLevelAxisLabels = bInDisplaySoundLevelAxisLabels; }
 	void SetDisplayFrequencyGridLines(bool bInDisplayFrequencyGridLines) { bDisplayFrequencyGridLines = bInDisplayFrequencyGridLines; }
@@ -201,7 +217,10 @@ private:
 	FLinearColor GetAxisLabelColor(const FWidgetStyle& InWidgetStyle) const;
 	FLinearColor GetSpectrumColor(const FWidgetStyle& InWidgetStyle) const;	
 
+	static float GetTiltExponentValue(const EAudioSpectrumPlotTilt InTilt);
+
 	TSharedRef<SWidget> BuildDefaultContextMenu();
+	void BuildTiltSpectrumSubMenu(FMenuBuilder& SubMenu);
 	void BuildFrequencyAxisScaleSubMenu(FMenuBuilder& SubMenu);
 	void BuildFrequencyAxisPixelBucketModeSubMenu(FMenuBuilder& SubMenu);
 
@@ -213,6 +232,8 @@ private:
 	TAttribute<float> ViewMaxFrequency;
 	TAttribute<float> ViewMinSoundLevel;
 	TAttribute<float> ViewMaxSoundLevel;
+	TAttribute<float> TiltExponent;
+	TAttribute<float> TiltPivotFrequency;
 	TAttribute<bool> bDisplayFrequencyAxisLabels;
 	TAttribute<bool> bDisplaySoundLevelAxisLabels;
 	TAttribute<bool> bDisplayFrequencyGridLines;
diff --git a/Engine/Plugins/Runtime/Metasound/Source/MetasoundEditor/Private/SMetasoundFilterGraphNodes.cpp b/Engine/Plugins/Runtime/Metasound/Source/MetasoundEditor/Private/SMetasoundFilterGraphNodes.cpp
--- a/Engine/Plugins/Runtime/Metasound/Source/MetasoundEditor/Private/SMetasoundFilterGraphNodes.cpp
+++ b/Engine/Plugins/Runtime/Metasound/Source/MetasoundEditor/Private/SMetasoundFilterGraphNodes.cpp
@@ -124,6 +124,7 @@ namespace Metasound::Editor
 			FrequencyResponsePlot = SNew(SAudioSpectrumPlot)
 				.Clipping(EWidgetClipping::ClipToBounds)
 				.ViewMinSoundLevel(-24.0f)
+				.TiltExponent_Lambda([] { return 0.0f; }) // Binding this property has the effect of hiding its context menu entry (Tilting the spectrum is not desired here).
 				.FrequencyAxisPixelBucketMode_Lambda([]() { return EAudioSpectrumPlotFrequencyAxisPixelBucketMode::Sample; }) // Binding this property has the effect of hiding its context menu entry (PixelBucketMode is not much use here).
 				.OnGetAudioSpectrumData(this, &SMetaSoundFilterGraphNode::GetAudioSpectrumData);
 
