UE - 211725 - AudioMaterialSlider is stretched out in Metasound Editor
#rnx

diff --git a/Engine/Plugins/Runtime/AudioWidgets/Source/AudioWidgets/Private/AudioMaterialSlate/AudioMaterialSlateTypes.cpp b/Engine/Plugins/Runtime/AudioWidgets/Source/AudioWidgets/Private/AudioMaterialSlate/AudioMaterialSlateTypes.cpp
--- a/Engine/Plugins/Runtime/AudioWidgets/Source/AudioWidgets/Private/AudioMaterialSlate/AudioMaterialSlateTypes.cpp
+++ b/Engine/Plugins/Runtime/AudioWidgets/Source/AudioWidgets/Private/AudioMaterialSlate/AudioMaterialSlateTypes.cpp
@@ -21,7 +21,7 @@ namespace AudioWidgets
 
 		static const FLinearColor BarMainColor(0.008f, 0.008f, 0.008f, 1.f);
 		static const FLinearColor BarAccentColor(0.005f, 0.005f, 0.005f, 1.f);
-		static const FLinearColor HandleMainColor(0.2f, 0.2f, 0.25f, 1.f);
+		static const FLinearColor HandleMainColor(1.f, 1.f, 1.f, 1.f);
 		static const FLinearColor HandleOutlineColor(0.15f, 0.15f, 0.15f, 1.f);
 	}
 }
@@ -81,7 +81,7 @@ FAudioMaterialSliderStyle::FAudioMaterialSliderStyle()
 	FString Path = PLUGIN_BASE_DIR + "MI_AudioMaterialRoundedSlider.MI_AudioMaterialRoundedSlider";
 	Material = LoadObject<UMaterialInterface>(nullptr, *Path);
 
-	DesiredSize = FVector2f(30.f, 250.f);
+	DesiredSize = FVector2f(25.f, 250.f);
 }
 
 const FName FAudioMaterialSliderStyle::TypeName(TEXT("FAudioMaterialSliderStyle"));
diff --git a/Engine/Plugins/Runtime/AudioWidgets/Source/AudioWidgets/Private/AudioMaterialSlate/SAudioMaterialLabeledSlider.cpp b/Engine/Plugins/Runtime/AudioWidgets/Source/AudioWidgets/Private/AudioMaterialSlate/SAudioMaterialLabeledSlider.cpp
--- a/Engine/Plugins/Runtime/AudioWidgets/Source/AudioWidgets/Private/AudioMaterialSlate/SAudioMaterialLabeledSlider.cpp
+++ b/Engine/Plugins/Runtime/AudioWidgets/Source/AudioWidgets/Private/AudioMaterialSlate/SAudioMaterialLabeledSlider.cpp
@@ -104,11 +104,11 @@ FVector2D SAudioMaterialLabeledSlider::ComputeDesiredSize(float) const
 	{
 		return DesiredSizeOverride.Get().GetValue();
 	}
-	
-	if (Style)
+
+	if (Style && Label.IsValid())
 	{
-		const float Width = Orientation.Get() == Orient_Vertical ? Style->DesiredSize.X : Style->DesiredSize.Y;
-		const float Heigth = Orientation.Get() == Orient_Vertical ? Style->DesiredSize.Y : Style->DesiredSize.X;
+		const float Width = Orientation.Get() == Orient_Vertical ? Label.Get()->GetDesiredSize().X + 6.f : Style->DesiredSize.Y + Label.Get()->GetDesiredSize().X;
+		const float Heigth = Orientation.Get() == Orient_Vertical ? Style->DesiredSize.Y + Label.Get()->GetDesiredSize().Y + 3.f: Style->DesiredSize.X;
 		return FVector2D(Width, Heigth);
 	}
 
@@ -209,6 +209,7 @@ TSharedRef<SWidgetSwitcher> SAudioMaterialLabeledSlider::CreateWidgetLayout()
 				// SSlider
 				.HAlign(HAlign_Fill)
 				.VAlign(VAlign_Center)
+				.Padding(3.0f, 0.0f)
 				[
 					Slider.ToSharedRef()
 				]
@@ -218,6 +219,7 @@ TSharedRef<SWidgetSwitcher> SAudioMaterialLabeledSlider::CreateWidgetLayout()
 			.AutoWidth()
 			.HAlign(HAlign_Fill)
 			.VAlign(VAlign_Center)
+			.Padding(3.0f, 0.0f, 0.0f, 0.0f)
 			[
 				Label.ToSharedRef()
 			]
@@ -234,6 +236,7 @@ TSharedRef<SWidgetSwitcher> SAudioMaterialLabeledSlider::CreateWidgetLayout()
 			SNew(SVerticalBox)
 			// Text Label
 			+ SVerticalBox::Slot()
+			.Padding(0.0f, 0.0f, 0.0f, 3.0f)
 			.AutoHeight()
 			.HAlign(HAlign_Center)
 			.VAlign(VAlign_Center)
@@ -249,6 +252,7 @@ TSharedRef<SWidgetSwitcher> SAudioMaterialLabeledSlider::CreateWidgetLayout()
 				// Actual SSlider
 				.HAlign(HAlign_Center)
 				.VAlign(VAlign_Fill)
+				.Padding(0.0f, 3.0f)
 				[
 					Slider.ToSharedRef()
 				]
diff --git a/Engine/Plugins/Runtime/AudioWidgets/Source/AudioWidgets/Public/AudioMaterialSlate/SAudioMaterialLabeledSlider.h b/Engine/Plugins/Runtime/AudioWidgets/Source/AudioWidgets/Public/AudioMaterialSlate/SAudioMaterialLabeledSlider.h
--- a/Engine/Plugins/Runtime/AudioWidgets/Source/AudioWidgets/Public/AudioMaterialSlate/SAudioMaterialLabeledSlider.h
+++ b/Engine/Plugins/Runtime/AudioWidgets/Source/AudioWidgets/Public/AudioMaterialSlate/SAudioMaterialLabeledSlider.h
@@ -4,6 +4,7 @@
 
 #include "AudioMaterialSlate/AudioMaterialSlateTypes.h"
 #include "AudioWidgetsEnums.h"
+#include "AudioWidgetsStyle.h"
 #include "Framework/SlateDelegates.h"
 #include "SAudioInputWidget.h"
 #include "Styling/ISlateStyle.h"
@@ -22,6 +23,7 @@ class AUDIOWIDGETS_API SAudioMaterialLabeledSlider
 {
 public:
 	SLATE_BEGIN_ARGS(SAudioMaterialLabeledSlider)
+	: _Style(&FAudioWidgetsStyle::Get().GetWidgetStyle<FAudioMaterialSliderStyle>("AudioMaterialSlider.Style"))
 	{}
 
 	/** The style used to draw the slider. */
diff --git a/Engine/Plugins/Runtime/Metasound/Source/MetasoundEditor/Private/SMetasoundGraphNode.cpp b/Engine/Plugins/Runtime/Metasound/Source/MetasoundEditor/Private/SMetasoundGraphNode.cpp
--- a/Engine/Plugins/Runtime/Metasound/Source/MetasoundEditor/Private/SMetasoundGraphNode.cpp
+++ b/Engine/Plugins/Runtime/Metasound/Source/MetasoundEditor/Private/SMetasoundGraphNode.cpp
@@ -833,7 +833,11 @@ namespace Metasound
 									[
 										InputWidget.ToSharedRef()
 									];
-								InputWidget->SetDesiredSizeOverride(SliderDesiredSizeVertical);
+
+									if (!bUseAudioMaterialWidgets)
+									{
+										InputWidget->SetDesiredSizeOverride(SliderDesiredSizeVertical);
+									}
 							}
 							else // horizontal orientation
 							{
@@ -867,7 +871,11 @@ namespace Metasound
 									[
 										Slot2.ToSharedRef()
 									];
-									InputWidget->SetDesiredSizeOverride(FVector2D(SliderDesiredSizeVertical.Y, SliderDesiredSizeVertical.X));
+
+									if (!bUseAudioMaterialWidgets)
+									{
+										InputWidget->SetDesiredSizeOverride(SliderDesiredSizeVertical);
+									}
 								}
 
 								if (bUseAudioMaterialWidgets)
