Added optional crosshair rendering to spectrum plot.

diff --git a/Engine/Plugins/Runtime/AudioWidgets/Source/AudioWidgets/Private/AudioSpectrumAnalyzer.cpp b/Engine/Plugins/Runtime/AudioWidgets/Source/AudioWidgets/Private/AudioSpectrumAnalyzer.cpp
--- a/Engine/Plugins/Runtime/AudioWidgets/Source/AudioWidgets/Private/AudioSpectrumAnalyzer.cpp
+++ b/Engine/Plugins/Runtime/AudioWidgets/Source/AudioWidgets/Private/AudioSpectrumAnalyzer.cpp
@@ -48,6 +48,7 @@ namespace AudioWidgets
 		, ConstantQSettings(NewObject<UConstantQSettings>())
 		, Widget(SNew(SAudioSpectrumPlot)
 			.Clipping(EWidgetClipping::ClipToBounds)
+			.DisplayCrosshair(true)
 			.DisplayFrequencyAxisLabels(false)
 			.DisplaySoundLevelAxisLabels(false)
 			.OnGetAudioSpectrumData_Raw(this, &FAudioSpectrumAnalyzer::GetAudioSpectrumData))
diff --git a/Engine/Plugins/Runtime/AudioWidgets/Source/AudioWidgets/Private/SAudioSpectrumPlot.cpp b/Engine/Plugins/Runtime/AudioWidgets/Source/AudioWidgets/Private/SAudioSpectrumPlot.cpp
--- a/Engine/Plugins/Runtime/AudioWidgets/Source/AudioWidgets/Private/SAudioSpectrumPlot.cpp
+++ b/Engine/Plugins/Runtime/AudioWidgets/Source/AudioWidgets/Private/SAudioSpectrumPlot.cpp
@@ -22,15 +22,30 @@ public:
 	void DrawSoundLevelGridLines(const int32 LayerId, TConstArrayView<float> GridLineSoundLevels, const FLinearColor& LineColor) const;
 	void DrawFrequencyGridLines(const int32 LayerId, TConstArrayView<float> GridLineFrequencies, const FLinearColor& LineColor) const;
 
+	void DrawCrosshairWithLabels(const int32 LayerId, float Frequency, float SoundLevel, const FSlateFontInfo& Font, const FLinearColor& TextColor, const FLinearColor& LineColor);
+
 	void DrawSoundLevelAxisLabels(const int32 LayerId, TConstArrayView<float> GridLineSoundLevels, const FSlateFontInfo& Font, const FLinearColor& TextColor);
 	void DrawFrequencyAxisLabels(const int32 LayerId, TConstArrayView<float> GridLineFrequencies, const FSlateFontInfo& Font, const FLinearColor& TextColor);
 
 	bool HasDrawnLabels() const { return !DrawnLabelRects.IsEmpty(); }
 
 private:
-	static FString FormatSoundLevelString(const float SoundLevel);
-	static FString FormatFreqString(const float Freq);
+	struct FSoundLevelFormattingOptions
+	{
+		int32 NumFractionalDigits = 0;
+		bool bIncludeUnits = false;
+	};
+
+	struct FFreqFormattingOptions
+	{
+		bool bAlwaysDisplayMaximumFractionalDigits = false;
+		bool bIncludeUnits = false;
+	};
 
+	static FString FormatSoundLevelString(const float SoundLevel, const FSoundLevelFormattingOptions& SoundLevelFormattingOptions);
+	static FString FormatFreqString(const float Freq, const FFreqFormattingOptions& FreqFormattingOptions);
+
+	void DrawVerticalArrowhead(int32 LayerId, FVector2f TipPosition, FVector2f Size, const FLinearColor& LineColor) const;
 	void DrawLabelIfNoOverlap(const int32 LayerId, const float LabelLeft, const float LabelTop, const FVector2f& LabelDrawSize, const FString LabelText, const FSlateFontInfo& Font, const FLinearColor& TextColor);
 
 	// Tweak the label Rect bounds to give space where it's needed for readability, while not wasting space where it's not needed.
@@ -88,13 +103,167 @@ void FAudioSpectrumPlotGridAndLabelDrawingHelper::DrawFrequencyGridLines(const i
 	}
 }
 
+void FAudioSpectrumPlotGridAndLabelDrawingHelper::DrawCrosshairWithLabels(const int32 LayerId, float Frequency, float SoundLevel, const FSlateFontInfo& Font, const FLinearColor& TextColor, const FLinearColor& LineColor)
+{
+	SpaceDrawSize = FontMeasureService->Measure(TEXT(" "), Font);
+
+	const float CrosshairPosX = ScaleInfo.FrequencyToLocalX(Frequency);
+	const float CrosshairPosY = ScaleInfo.SoundLevelToLocalY(SoundLevel);
+	const bool bIsHorizontalCrosshairWithinVisibleRange = (CrosshairPosY >= LocalBackgroundRect.Top && CrosshairPosY <= LocalBackgroundRect.Bottom);
+	const bool bIsVerticalCrosshairWithinVisibleRange = (CrosshairPosX >= LocalBackgroundRect.Left && CrosshairPosX <= LocalBackgroundRect.Right);
+
+	if (!bIsVerticalCrosshairWithinVisibleRange)
+	{
+		return;
+	}
+
+	// If the horizontal crosshair is not within visible range then we shall be drawing arrowheads at the top or bottom to signify this.
+	constexpr float ArrowheadHeight = 4.0f;
+	constexpr float ArrowheadWidth = 6.0f;
+	const float ArrowTipPosY = FMath::Clamp(CrosshairPosY, LocalBackgroundRect.Top, LocalBackgroundRect.Bottom);
+	const float ArrowDirection = FMath::Sign(CrosshairPosY - ArrowTipPosY);
+
+	TArray<FVector2f> VerticalCrosshairLinePoints;
+	VerticalCrosshairLinePoints.Reserve(2);
+
+	const FString FreqString = FormatFreqString(Frequency, { .bAlwaysDisplayMaximumFractionalDigits = true, .bIncludeUnits = true });
+	const FVector2f FreqLabelDrawSize = FontMeasureService->Measure(FreqString, Font);
+	const float FreqLabelLeft = FMath::Clamp(CrosshairPosX - 0.5f * FreqLabelDrawSize.X, LocalBackgroundRect.Left, LocalBackgroundRect.Right - FreqLabelDrawSize.X);
+	const float TopLabelBottomSide = LocalBackgroundRect.Top + FreqLabelDrawSize.Y;
+	const float BottomLabelTopSide = LocalBackgroundRect.Bottom - FreqLabelDrawSize.Y;
+
+	if (CrosshairPosY >= TopLabelBottomSide)
+	{
+		// Draw label at the top:
+		DrawLabelIfNoOverlap(LayerId, FreqLabelLeft, LocalBackgroundRect.Top, FreqLabelDrawSize, FreqString, Font, TextColor);
+
+		// Start the vertical crosshair line below the top label:
+		VerticalCrosshairLinePoints.Add({ CrosshairPosX, TopLabelBottomSide });
+	}
+	else
+	{
+		// Don't draw label at the top, as either the horizontal crosshair line or an arrowhead will be drawn at the top.
+
+		// Start the vertical crosshair line at the very top:
+		VerticalCrosshairLinePoints.Add({ CrosshairPosX, LocalBackgroundRect.Top });
+	}
+
+	if (CrosshairPosY <= BottomLabelTopSide)
+	{
+		// Draw label at the bottom:
+		DrawLabelIfNoOverlap(LayerId, FreqLabelLeft, BottomLabelTopSide, FreqLabelDrawSize, FreqString, Font, TextColor);
+
+		// End the crosshair line above the bottom label:
+		VerticalCrosshairLinePoints.Add({ CrosshairPosX, BottomLabelTopSide });
+	}
+	else
+	{
+		// Don't draw label at the bottom, as either the horizontal crosshair line or an arrowhead will be drawn at the bottom.
+
+		// End the vertical crosshair line at the very bottom:
+		VerticalCrosshairLinePoints.Add({ CrosshairPosX, LocalBackgroundRect.Bottom });
+	}
+
+	// Draw the vertical crosshair:
+	FSlateDrawElement::MakeLines(ElementList, LayerId, AllottedGeometry.ToPaintGeometry(), VerticalCrosshairLinePoints, ESlateDrawEffect::None, LineColor);
+
+	if (!bIsHorizontalCrosshairWithinVisibleRange)
+	{
+		// Draw an arrowhead at the top or bottom of the vertical crosshair:
+		DrawVerticalArrowhead(LayerId, { CrosshairPosX, ArrowTipPosY }, { ArrowheadWidth, ArrowDirection * ArrowheadHeight }, LineColor);
+	}
+
+
+
+	TArray<FVector2f> HorizontalCrosshairLinePoints;
+	HorizontalCrosshairLinePoints.Reserve(2);
+
+	const FString SoundLevelString = FormatSoundLevelString(SoundLevel, { .NumFractionalDigits = 1, .bIncludeUnits = true });
+	const FVector2f SoundLevelLabelDrawSize = FontMeasureService->Measure(SoundLevelString, Font);
+	const float SoundLevelLabelTop = FMath::Clamp(CrosshairPosY - 0.5f * SoundLevelLabelDrawSize.Y, LocalBackgroundRect.Top, LocalBackgroundRect.Bottom - SoundLevelLabelDrawSize.Y);
+	const float LeftLabelRightSide = LocalBackgroundRect.Left + SoundLevelLabelDrawSize.X;
+	const float RightLabelLeftSide = LocalBackgroundRect.Right - SoundLevelLabelDrawSize.X;
+	const float LeftLabelRightSidePadded = LeftLabelRightSide + SpaceDrawSize.X;
+	const float RightLabelLeftSidePadded = RightLabelLeftSide - SpaceDrawSize.X;
+
+	if (bIsHorizontalCrosshairWithinVisibleRange)
+	{
+		if (CrosshairPosX > LeftLabelRightSide)
+		{
+			// Draw label at the left end of the horizontal crosshair line:
+			DrawLabelIfNoOverlap(LayerId, LocalBackgroundRect.Left, SoundLevelLabelTop, SoundLevelLabelDrawSize, SoundLevelString, Font, TextColor);
+
+			// Start the horizontal crosshair line to the right of the left side label:
+			HorizontalCrosshairLinePoints.Add({ FMath::Min(LeftLabelRightSidePadded, CrosshairPosX), CrosshairPosY });
+		}
+		else
+		{
+			// Start the horizontal crosshair line at the furthest left:
+			HorizontalCrosshairLinePoints.Add({ LocalBackgroundRect.Left, CrosshairPosY });
+		}
+
+		if (CrosshairPosX < RightLabelLeftSide)
+		{
+			// Draw label at the right end of the horizontal crosshair line:
+			DrawLabelIfNoOverlap(LayerId, RightLabelLeftSide, SoundLevelLabelTop, SoundLevelLabelDrawSize, SoundLevelString, Font, TextColor);
+
+			// End the horizontal crosshair line to the left of the right side label:
+			HorizontalCrosshairLinePoints.Add({ FMath::Max(RightLabelLeftSidePadded, CrosshairPosX), CrosshairPosY });
+		}
+		else
+		{
+			// End the horizontal crosshair line at the furthest right:
+			HorizontalCrosshairLinePoints.Add({ LocalBackgroundRect.Right, CrosshairPosY });
+		}
+
+		// Draw the horizontal crosshair:
+		FSlateDrawElement::MakeLines(ElementList, LayerId, AllottedGeometry.ToPaintGeometry(), HorizontalCrosshairLinePoints, ESlateDrawEffect::None, LineColor);
+	}
+	else
+	{
+		const float ArrowTailPosY = SoundLevelLabelTop + 0.5f * SoundLevelLabelDrawSize.Y;
+
+		if (CrosshairPosX > LeftLabelRightSidePadded + 1.5f * ArrowheadWidth)
+		{
+			// Draw label in the left corner:
+			DrawLabelIfNoOverlap(LayerId, LocalBackgroundRect.Left, SoundLevelLabelTop, SoundLevelLabelDrawSize, SoundLevelString, Font, TextColor);
+
+			// Horizontal crosshair is out of visible range, draw a vertical arrow to the right of the sound level label to signify this:
+			const float ArrowPosX = LeftLabelRightSidePadded + 0.5f * ArrowheadWidth;
+			DrawVerticalArrowhead(LayerId, { ArrowPosX, ArrowTipPosY }, { ArrowheadWidth, ArrowDirection * ArrowheadHeight }, LineColor);
+
+			TArray<FVector2f> ArrowShaftPoints;
+			ArrowShaftPoints.Reserve(2);
+			ArrowShaftPoints.Add({ ArrowPosX, ArrowTailPosY });
+			ArrowShaftPoints.Add({ ArrowPosX, ArrowTipPosY });
+			FSlateDrawElement::MakeLines(ElementList, LayerId, AllottedGeometry.ToPaintGeometry(), ArrowShaftPoints, ESlateDrawEffect::None, LineColor);
+		}
+
+		if (CrosshairPosX < RightLabelLeftSidePadded - 1.5f * ArrowheadWidth)
+		{
+			// Draw label in the right corner:
+			DrawLabelIfNoOverlap(LayerId, RightLabelLeftSide, SoundLevelLabelTop, SoundLevelLabelDrawSize, SoundLevelString, Font, TextColor);
+
+			// Horizontal crosshair is out of visible range, draw a vertical arrow to the left of the sound level label to signify this:
+			const float ArrowPosX = RightLabelLeftSidePadded - 0.5f * ArrowheadWidth;
+			DrawVerticalArrowhead(LayerId, { ArrowPosX, ArrowTipPosY }, { ArrowheadWidth, ArrowDirection * ArrowheadHeight }, LineColor);
+
+			TArray<FVector2f> ArrowShaftPoints;
+			ArrowShaftPoints.Reserve(2);
+			ArrowShaftPoints.Add({ ArrowPosX, ArrowTailPosY });
+			ArrowShaftPoints.Add({ ArrowPosX, ArrowTipPosY });
+			FSlateDrawElement::MakeLines(ElementList, LayerId, AllottedGeometry.ToPaintGeometry(), ArrowShaftPoints, ESlateDrawEffect::None, LineColor);
+		}
+	}
+}
+
 void FAudioSpectrumPlotGridAndLabelDrawingHelper::DrawSoundLevelAxisLabels(const int32 LayerId, TConstArrayView<float> GridLineSoundLevels, const FSlateFontInfo& Font, const FLinearColor& TextColor)
 {
 	SpaceDrawSize = FontMeasureService->Measure(TEXT(" "), Font);
 
 	for (const float SoundLevel : GridLineSoundLevels)
 	{
-		const FString SoundLevelString = FormatSoundLevelString(SoundLevel);
+		const FString SoundLevelString = FormatSoundLevelString(SoundLevel, { .NumFractionalDigits = 0, .bIncludeUnits = false });
 		const FVector2f LabelDrawSize = FontMeasureService->Measure(SoundLevelString, Font);
 		const float GridLineLocalY = ScaleInfo.SoundLevelToLocalY(SoundLevel);
 		const float LabelTop = GridLineLocalY - 0.5f * LabelDrawSize.Y;
@@ -116,7 +285,7 @@ void FAudioSpectrumPlotGridAndLabelDrawingHelper::DrawFrequencyAxisLabels(const
 
 	for (float Freq : GridLineFrequencies)
 	{
-		const FString FreqString = FormatFreqString(Freq);
+		const FString FreqString = FormatFreqString(Freq, { .bAlwaysDisplayMaximumFractionalDigits = false, .bIncludeUnits = false });
 		const FVector2f LabelDrawSize = FontMeasureService->Measure(FreqString, Font);
 		const float GridLineLocalX = ScaleInfo.FrequencyToLocalX(Freq);
 		const float LabelLeft = GridLineLocalX - 0.5f * LabelDrawSize.X;
@@ -132,30 +301,76 @@ void FAudioSpectrumPlotGridAndLabelDrawingHelper::DrawFrequencyAxisLabels(const
 	}
 }
 
-FString FAudioSpectrumPlotGridAndLabelDrawingHelper::FormatSoundLevelString(const float SoundLevel)
+FString FAudioSpectrumPlotGridAndLabelDrawingHelper::FormatSoundLevelString(const float SoundLevel, const FSoundLevelFormattingOptions& SoundLevelFormattingOptions)
 {
-	if (SoundLevel > 0.0f)
+	FNumberFormattingOptions NumberFormattingOptions;
+	NumberFormattingOptions.MinimumFractionalDigits = SoundLevelFormattingOptions.NumFractionalDigits;
+	NumberFormattingOptions.MaximumFractionalDigits = SoundLevelFormattingOptions.NumFractionalDigits;
+	if (SoundLevel != 0.0f)
+	{
+		NumberFormattingOptions.AlwaysSign = true;
+	}
+
+	const FText NumberText = FText::AsNumber(SoundLevel, &NumberFormattingOptions);
+	if (SoundLevelFormattingOptions.bIncludeUnits)
 	{
-		return FString::Printf(TEXT("+%.0f"), SoundLevel);
+		return NumberText.ToString() + TEXT(" dB");
 	}
 	else
 	{
-		return FString::Printf(TEXT("%.0f"), SoundLevel);
+		return NumberText.ToString();
 	}
 }
 
-FString FAudioSpectrumPlotGridAndLabelDrawingHelper::FormatFreqString(const float Freq)
+FString FAudioSpectrumPlotGridAndLabelDrawingHelper::FormatFreqString(const float Freq, const FFreqFormattingOptions& FreqFormattingOptions)
 {
+	FNumberFormattingOptions NumberFormattingOptions;
+
 	if (Freq >= 1000.0f)
 	{
-		return FString::Printf(TEXT("%.0fk"), Freq / 1000.0f);
+		NumberFormattingOptions.MaximumFractionalDigits = (Freq < 10000.0f) ? 2 : 1; // Displaying a max of 3 significant figures.
+		if (FreqFormattingOptions.bAlwaysDisplayMaximumFractionalDigits)
+		{
+			NumberFormattingOptions.MinimumFractionalDigits = NumberFormattingOptions.MaximumFractionalDigits;
+		}
+
+		const FText NumberText = FText::AsNumber(Freq / 1000.0f, &NumberFormattingOptions);
+		if (FreqFormattingOptions.bIncludeUnits)
+		{
+			return NumberText.ToString() + TEXT(" kHz");
+		}
+		else
+		{
+			return NumberText.ToString() + TEXT(" k");
+		}
 	}
 	else
 	{
-		return FString::Printf(TEXT("%.0f"), Freq);
+		NumberFormattingOptions.MaximumFractionalDigits = 0;
+		NumberFormattingOptions.MinimumFractionalDigits = 0;
+
+		const FText NumberText = FText::AsNumber(Freq, &NumberFormattingOptions);
+		if (FreqFormattingOptions.bIncludeUnits)
+		{
+			return NumberText.ToString() + TEXT(" Hz");
+		}
+		else
+		{
+			return NumberText.ToString();
+		}
 	}
 }
 
+void FAudioSpectrumPlotGridAndLabelDrawingHelper::DrawVerticalArrowhead(int32 LayerId, FVector2f TipPosition, FVector2f Size, const FLinearColor& LineColor) const
+{
+	TArray<FVector2f> LinePoints;
+	LinePoints.Reserve(3);
+	LinePoints.Add({ TipPosition.X - 0.5f * Size.X, TipPosition.Y - Size.Y });
+	LinePoints.Add({ TipPosition.X, TipPosition.Y });
+	LinePoints.Add({ TipPosition.X + 0.5f * Size.X, TipPosition.Y - Size.Y });
+	FSlateDrawElement::MakeLines(ElementList, LayerId, AllottedGeometry.ToPaintGeometry(), LinePoints, ESlateDrawEffect::None, LineColor);
+}
+
 void FAudioSpectrumPlotGridAndLabelDrawingHelper::DrawLabelIfNoOverlap(const int32 LayerId, const float LabelLeft, const float LabelTop, const FVector2f& LabelDrawSize, const FString LabelText, const FSlateFontInfo& Font, const FLinearColor& TextColor)
 {
 	const FSlateLayoutTransform LabelTransform(FVector2f(LabelLeft, LabelTop));
@@ -191,7 +406,7 @@ bool FAudioSpectrumPlotGridAndLabelDrawingHelper::IsOverlappingPreviouslyDrawnLa
 	return (OverlappingDrawnLabel != nullptr);
 };
 
-
+const float SAudioSpectrumPlot::ClampMinSoundLevel = -200.0f;
 FName SAudioSpectrumPlot::ContextMenuExtensionHook("SpectrumPlotDisplayOptions");
 
 void SAudioSpectrumPlot::Construct(const FArguments& InArgs)
@@ -205,6 +420,8 @@ void SAudioSpectrumPlot::Construct(const FArguments& InArgs)
 	ViewMaxSoundLevel = InArgs._ViewMaxSoundLevel;
 	TiltExponent = InArgs._TiltExponent;
 	TiltPivotFrequency = InArgs._TiltPivotFrequency;
+	SelectedFrequency = InArgs._SelectedFrequency;
+	bDisplayCrosshair = InArgs._DisplayCrosshair;
 	bDisplayFrequencyAxisLabels = InArgs._DisplayFrequencyAxisLabels;
 	bDisplaySoundLevelAxisLabels = InArgs._DisplaySoundLevelAxisLabels;
 	bDisplayFrequencyGridLines = InArgs._DisplayFrequencyGridLines;
@@ -214,6 +431,7 @@ void SAudioSpectrumPlot::Construct(const FArguments& InArgs)
 	BackgroundColor = InArgs._BackgroundColor;
 	GridColor = InArgs._GridColor;
 	AxisLabelColor = InArgs._AxisLabelColor;
+	CrosshairColor = InArgs._CrosshairColor;
 	SpectrumColor = InArgs._SpectrumColor;
 	bAllowContextMenu = InArgs._AllowContextMenu;
 	OnContextMenuOpening = InArgs._OnContextMenuOpening;
@@ -283,6 +501,28 @@ FReply SAudioSpectrumPlot::OnMouseButtonUp(const FGeometry& InMyGeometry, const
 	return SCompoundWidget::OnMouseButtonUp(InMyGeometry, InMouseEvent);
 }
 
+FReply SAudioSpectrumPlot::OnMouseMove(const FGeometry& MyGeometry, const FPointerEvent& MouseEvent)
+{
+	if (!SelectedFrequency.IsBound())
+	{
+		// If not bound to an external function, set SelectedFrequency from the mouse hover position:
+		const FAudioSpectrumPlotScaleInfo ScaleInfo(MyGeometry.GetLocalSize(), FrequencyAxisScale.Get(), ViewMinFrequency.Get(), ViewMaxFrequency.Get(), ViewMinSoundLevel.Get(), ViewMaxSoundLevel.Get());
+		const FVector2f MouseMoveLocation = MyGeometry.AbsoluteToLocal(MouseEvent.GetScreenSpacePosition());
+		SelectedFrequency = ScaleInfo.LocalXToFrequency(MouseMoveLocation.X);
+	}
+
+	return FReply::Unhandled();
+}
+
+void SAudioSpectrumPlot::OnMouseLeave(const FPointerEvent& MouseEvent)
+{
+	if (!SelectedFrequency.IsBound())
+	{
+		// If not bound to an external function, clear SelectedFrequency when mouse is no longer hovering over the plot widget:
+		SelectedFrequency = NullOpt;
+	}
+}
+
 FAudioSpectrumPlotScaleInfo SAudioSpectrumPlot::GetScaleInfo() const
 {
 	const FGeometry& AllottedGeometry = GetPaintSpaceGeometry();
@@ -295,7 +535,7 @@ int32 SAudioSpectrumPlot::OnPaint(const FPaintArgs& Args, const FGeometry& Allot
 
 	LayerId = DrawSolidBackgroundRectangle(AllottedGeometry, OutDrawElements, LayerId, InWidgetStyle);
 
-	LayerId = DrawGridAndLabels(AllottedGeometry, OutDrawElements, LayerId, InWidgetStyle, ScaleInfo);
+	LayerId = DrawGrid(AllottedGeometry, OutDrawElements, LayerId, InWidgetStyle, ScaleInfo);
 
 	LayerId = DrawPowerSpectrum(AllottedGeometry, OutDrawElements, LayerId, InWidgetStyle, ScaleInfo);
 
@@ -311,7 +551,7 @@ int32 SAudioSpectrumPlot::DrawSolidBackgroundRectangle(const FGeometry& Allotted
 	return LayerId + 1;
 }
 
-int32 SAudioSpectrumPlot::DrawGridAndLabels(const FGeometry& AllottedGeometry, FSlateWindowElementList& OutDrawElements, int32 LayerId, const FWidgetStyle& InWidgetStyle, const FAudioSpectrumPlotScaleInfo& ScaleInfo) const
+int32 SAudioSpectrumPlot::DrawGrid(const FGeometry& AllottedGeometry, FSlateWindowElementList& OutDrawElements, int32 LayerId, const FWidgetStyle& InWidgetStyle, const FAudioSpectrumPlotScaleInfo& ScaleInfo) const
 {
 	TArray<float> GridLineSoundLevels;
 	GetGridLineSoundLevels(GridLineSoundLevels);
@@ -320,9 +560,7 @@ int32 SAudioSpectrumPlot::DrawGridAndLabels(const FGeometry& AllottedGeometry, F
 	TArray<float> MajorGridLineFrequencies;
 	GetGridLineFrequencies(AllGridLineFrequencies, MajorGridLineFrequencies);
 
-	const FSlateFontInfo& Font = Style->AxisLabelFont;
 	const FLinearColor LineColor = GetGridColor(InWidgetStyle);
-	const FLinearColor TextColor = GetAxisLabelColor(InWidgetStyle);
 
 	FAudioSpectrumPlotGridAndLabelDrawingHelper GridAndLabelDrawingHelper(AllottedGeometry, OutDrawElements, ScaleInfo);
 
@@ -335,27 +573,8 @@ int32 SAudioSpectrumPlot::DrawGridAndLabels(const FGeometry& AllottedGeometry, F
 	{
 		GridAndLabelDrawingHelper.DrawFrequencyGridLines(LayerId, AllGridLineFrequencies, LineColor);
 	}
-	LayerId++;
-
-	if (bDisplaySoundLevelAxisLabels.Get())
-	{
-		// Draw sound level axis labels for all grid lines.
-		GridAndLabelDrawingHelper.DrawSoundLevelAxisLabels(LayerId, GridLineSoundLevels, Font, TextColor);
-	}
 
-	if (bDisplayFrequencyAxisLabels.Get())
-	{
-		// Draw frequency axis labels for all major grid lines.
-		GridAndLabelDrawingHelper.DrawFrequencyAxisLabels(LayerId, MajorGridLineFrequencies, Font, TextColor);
-	}
-
-	if (GridAndLabelDrawingHelper.HasDrawnLabels())
-	{
-		// We drew some labels, so increment layer ID:
-		LayerId++;
-	}
-
-	return LayerId;
+	return LayerId + 1;
 }
 
 void SAudioSpectrumPlot::GetGridLineSoundLevels(TArray<float>& GridLineSoundLevels) const
@@ -447,7 +666,7 @@ int32 SAudioSpectrumPlot::DrawPowerSpectrum(const FGeometry& AllottedGeometry, F
 		const float TiltExponentValue = TiltExponent.Get();
 		const float TiltPivotFrequencyValue = TiltPivotFrequency.Get();
 		const float ClampMinFrequency = (FrequencyAxisScale.Get() == EAudioSpectrumPlotFrequencyAxisScale::Logarithmic) ? 0.00001f : -FLT_MAX; // Cannot plot DC with log scale.
-		const float ClampMinMagnitudeSquared = FMath::Pow(10.0f, -200.0f / 10.0f); // Clamp at -200dB
+		const float ClampMinMagnitudeSquared = FMath::Pow(10.0f, ClampMinSoundLevel / 10.0f); // Clamp at -200dB
 		for (int Index = 0; Index < NumFrequencies; Index++)
 		{
 			const float Frequency = FMath::Max(PowerSpectrum.CenterFrequencies[Index], ClampMinFrequency);
@@ -522,6 +741,9 @@ int32 SAudioSpectrumPlot::DrawPowerSpectrum(const FGeometry& AllottedGeometry, F
 			break;
 		}
 
+		// Draw crosshair and axis labels (horizontal crosshair position depends on the spectrum line points to be plotted):
+		LayerId = DrawCrosshairAndAxisLabels(AllottedGeometry, OutDrawElements, LayerId, InWidgetStyle, ScaleInfo, LinePoints);
+
 		// Actually draw the line points:
 		const FLinearColor& LineColor = GetSpectrumColor(InWidgetStyle);
 		FSlateDrawElement::MakeLines(OutDrawElements, LayerId, AllottedGeometry.ToPaintGeometry(), LinePoints, ESlateDrawEffect::None, LineColor, true, 1.0f);
@@ -531,6 +753,73 @@ int32 SAudioSpectrumPlot::DrawPowerSpectrum(const FGeometry& AllottedGeometry, F
 	return LayerId;
 }
 
+int32 SAudioSpectrumPlot::DrawCrosshairAndAxisLabels(const FGeometry& AllottedGeometry, FSlateWindowElementList& OutDrawElements, int32 LayerId, const FWidgetStyle& InWidgetStyle, const FAudioSpectrumPlotScaleInfo& ScaleInfo, TConstArrayView<FVector2f> LinePoints) const
+{
+	FAudioSpectrumPlotGridAndLabelDrawingHelper GridAndLabelDrawingHelper(AllottedGeometry, OutDrawElements, ScaleInfo);
+
+	FLinearColor TextColor = GetAxisLabelColor(InWidgetStyle);
+
+	if (bDisplayCrosshair.Get())
+	{
+		const TOptional<float> SelectedFrequencyOptional = SelectedFrequency.Get();
+		if (SelectedFrequencyOptional.IsSet())
+		{
+			const float CrosshairFrequency = SelectedFrequencyOptional.GetValue();
+			const float CrosshairPosX = ScaleInfo.FrequencyToLocalX(CrosshairFrequency);
+			const int32 PointIndex = LinePoints.FindLastByPredicate([=](FVector2f Point) { return Point.X <= CrosshairPosX; });
+			if (PointIndex != INDEX_NONE && PointIndex + 1 < LinePoints.Num())
+			{
+				const FVector2f& PointL = LinePoints[PointIndex];
+				const FVector2f& PointR = LinePoints[PointIndex + 1];
+				const float LerpParam = (CrosshairPosX - PointL.X) / (PointR.X - PointL.X);
+				const float CrosshairPosY = FMath::Lerp(PointL.Y, PointR.Y, LerpParam);
+				const float CrosshairSoundLevel = ScaleInfo.LocalYToSoundLevel(CrosshairPosY);
+				if (CrosshairSoundLevel > ClampMinSoundLevel)
+				{
+					const FSlateFontInfo& CrosshairLabelFont = Style->CrosshairLabelFont;
+					const FLinearColor CrosshairLineColor = GetCrosshairColor(InWidgetStyle);
+					GridAndLabelDrawingHelper.DrawCrosshairWithLabels(LayerId, CrosshairFrequency, CrosshairSoundLevel, CrosshairLabelFont, TextColor, CrosshairLineColor);
+				}
+			}
+		}
+	}
+
+	TArray<float> GridLineSoundLevels;
+	GetGridLineSoundLevels(GridLineSoundLevels);
+
+	TArray<float> AllGridLineFrequencies;
+	TArray<float> MajorGridLineFrequencies;
+	GetGridLineFrequencies(AllGridLineFrequencies, MajorGridLineFrequencies);
+
+	const FSlateFontInfo& AxisLabelFont = Style->AxisLabelFont;
+
+	if (GridAndLabelDrawingHelper.HasDrawnLabels())
+	{
+		// De-emphasize grid axis labels if we are displaying crosshair labels:
+		TextColor.A *= 0.5;
+	}
+
+	if (bDisplaySoundLevelAxisLabels.Get())
+	{
+		// Draw sound level axis labels for all grid lines.
+		GridAndLabelDrawingHelper.DrawSoundLevelAxisLabels(LayerId, GridLineSoundLevels, AxisLabelFont, TextColor);
+	}
+
+	if (bDisplayFrequencyAxisLabels.Get())
+	{
+		// Draw frequency axis labels for all major grid lines.
+		GridAndLabelDrawingHelper.DrawFrequencyAxisLabels(LayerId, MajorGridLineFrequencies, AxisLabelFont, TextColor);
+	}
+
+	if (GridAndLabelDrawingHelper.HasDrawnLabels())
+	{
+		// We drew some labels, so increment layer ID:
+		LayerId++;
+	}
+
+	return LayerId;
+}
+
 FAudioPowerSpectrumData SAudioSpectrumPlot::GetPowerSpectrum() const
 {
 	if (OnGetAudioSpectrumData.IsBound())
@@ -599,6 +888,12 @@ FLinearColor SAudioSpectrumPlot::GetAxisLabelColor(const FWidgetStyle& InWidgetS
 	return SlateColor.GetColor(InWidgetStyle) * InWidgetStyle.GetColorAndOpacityTint();
 }
 
+FLinearColor SAudioSpectrumPlot::GetCrosshairColor(const FWidgetStyle& InWidgetStyle) const
+{
+	const FSlateColor& SlateColor = (CrosshairColor.Get() != FSlateColor::UseStyle()) ? CrosshairColor.Get() : Style->CrosshairColor;
+	return SlateColor.GetColor(InWidgetStyle) * InWidgetStyle.GetColorAndOpacityTint();
+}
+
 FLinearColor SAudioSpectrumPlot::GetSpectrumColor(const FWidgetStyle& InWidgetStyle) const
 {
 	const FSlateColor& SlateColor = (SpectrumColor.Get() != FSlateColor::UseStyle()) ? SpectrumColor.Get() : Style->SpectrumColor;
diff --git a/Engine/Plugins/Runtime/AudioWidgets/Source/AudioWidgets/Public/AudioSpectrumPlotStyle.h b/Engine/Plugins/Runtime/AudioWidgets/Source/AudioWidgets/Public/AudioSpectrumPlotStyle.h
--- a/Engine/Plugins/Runtime/AudioWidgets/Source/AudioWidgets/Public/AudioSpectrumPlotStyle.h
+++ b/Engine/Plugins/Runtime/AudioWidgets/Source/AudioWidgets/Public/AudioSpectrumPlotStyle.h
@@ -44,4 +44,12 @@ struct AUDIOWIDGETS_API FAudioSpectrumPlotStyle : public FSlateWidgetStyle
 	UPROPERTY(EditAnywhere, BlueprintReadWrite, Category = Appearance)
 	FSlateColor SpectrumColor = FStyleColors::Foreground;
 	FAudioSpectrumPlotStyle& SetSpectrumColor(const FSlateColor& InSpectrumColor) { SpectrumColor = InSpectrumColor; return *this; }
+
+	UPROPERTY(EditAnywhere, BlueprintReadWrite, Category = Appearance)
+	FSlateColor CrosshairColor = FStyleColors::Foreground;
+	FAudioSpectrumPlotStyle& SetCrosshairColor(const FSlateColor& InCrosshairColor) { CrosshairColor = InCrosshairColor; return *this; }
+
+	UPROPERTY(EditAnywhere, BlueprintReadWrite, Category = Appearance)
+	FSlateFontInfo CrosshairLabelFont = FStyleDefaults::GetFontInfo(8);
+	FAudioSpectrumPlotStyle& SetCrosshairLabelFont(const FSlateFontInfo& InCrosshairLabelFont) { CrosshairLabelFont = InCrosshairLabelFont; return *this; }
 };
diff --git a/Engine/Plugins/Runtime/AudioWidgets/Source/AudioWidgets/Public/SAudioSpectrumPlot.h b/Engine/Plugins/Runtime/AudioWidgets/Source/AudioWidgets/Public/SAudioSpectrumPlot.h
--- a/Engine/Plugins/Runtime/AudioWidgets/Source/AudioWidgets/Public/SAudioSpectrumPlot.h
+++ b/Engine/Plugins/Runtime/AudioWidgets/Source/AudioWidgets/Public/SAudioSpectrumPlot.h
@@ -131,6 +131,7 @@ public:
 		, _ViewMaxSoundLevel(12.0f)
 		, _TiltExponent(0.0f)
 		, _TiltPivotFrequency(24000.0f)
+		, _DisplayCrosshair(false)
 		, _DisplayFrequencyAxisLabels(true)
 		, _DisplaySoundLevelAxisLabels(true)
 		, _DisplayFrequencyGridLines(true)
@@ -140,6 +141,7 @@ public:
 		, _BackgroundColor(FSlateColor::UseStyle())
 		, _GridColor(FSlateColor::UseStyle())
 		, _AxisLabelColor(FSlateColor::UseStyle())
+		, _CrosshairColor(FSlateColor::UseStyle())
 		, _SpectrumColor(FSlateColor::UseStyle())
 		, _AllowContextMenu(true)
 	{}
@@ -150,6 +152,8 @@ public:
 		SLATE_ATTRIBUTE(float, ViewMaxSoundLevel)
 		SLATE_ATTRIBUTE(float, TiltExponent)
 		SLATE_ATTRIBUTE(float, TiltPivotFrequency)
+		SLATE_ATTRIBUTE(TOptional<float>, SelectedFrequency)
+		SLATE_ATTRIBUTE(bool, DisplayCrosshair)
 		SLATE_ATTRIBUTE(bool, DisplayFrequencyAxisLabels)
 		SLATE_ATTRIBUTE(bool, DisplaySoundLevelAxisLabels)
 		SLATE_ATTRIBUTE(bool, DisplayFrequencyGridLines)
@@ -159,6 +163,7 @@ public:
 		SLATE_ATTRIBUTE(FSlateColor, BackgroundColor)
 		SLATE_ATTRIBUTE(FSlateColor, GridColor)
 		SLATE_ATTRIBUTE(FSlateColor, AxisLabelColor)
+		SLATE_ATTRIBUTE(FSlateColor, CrosshairColor)
 		SLATE_ATTRIBUTE(FSlateColor, SpectrumColor)
 		SLATE_ATTRIBUTE(bool, AllowContextMenu)
 		SLATE_EVENT(FOnContextMenuOpening, OnContextMenuOpening)
@@ -174,6 +179,8 @@ public:
 	void SetViewMaxSoundLevel(float InViewMaxSoundLevel) { ViewMaxSoundLevel = InViewMaxSoundLevel; }
 	void SetTiltExponent(float InTiltExponent) { TiltExponent = InTiltExponent; }
 	void SetTiltPivotFrequency(float InTiltPivotFrequency) { TiltPivotFrequency = InTiltPivotFrequency; }
+	void SetSelectedFrequency(TOptional<float> InSelectedFrequency) { SelectedFrequency = InSelectedFrequency; }
+	void SetDisplayCrosshair(bool bInDisplayCrosshair) { bDisplayCrosshair = bInDisplayCrosshair; }
 	void SetDisplayFrequencyAxisLabels(bool bInDisplayFrequencyAxisLabels) { bDisplayFrequencyAxisLabels = bInDisplayFrequencyAxisLabels; }
 	void SetDisplaySoundLevelAxisLabels(bool bInDisplaySoundLevelAxisLabels) { bDisplaySoundLevelAxisLabels = bInDisplaySoundLevelAxisLabels; }
 	void SetDisplayFrequencyGridLines(bool bInDisplayFrequencyGridLines) { bDisplayFrequencyGridLines = bInDisplayFrequencyGridLines; }
@@ -188,6 +195,8 @@ public:
 	// Begin SWidget overrides.
 	virtual FReply OnMouseButtonDown(const FGeometry& InMyGeometry, const FPointerEvent& InMouseEvent) override;
 	virtual FReply OnMouseButtonUp(const FGeometry& InMyGeometry, const FPointerEvent& InMouseEvent) override;
+	virtual FReply OnMouseMove(const FGeometry& MyGeometry, const FPointerEvent& MouseEvent) override;
+	virtual void OnMouseLeave(const FPointerEvent& MouseEvent) override;
 	// End SWidget overrides.
 
 	void UnbindOnGetAudioSpectrumData() { OnGetAudioSpectrumData.Unbind(); }
@@ -201,13 +210,15 @@ private:
 
 	int32 DrawSolidBackgroundRectangle(const FGeometry& AllottedGeometry, FSlateWindowElementList& OutDrawElements, int32 LayerId, const FWidgetStyle& InWidgetStyle) const;
 
-	int32 DrawGridAndLabels(const FGeometry& AllottedGeometry, FSlateWindowElementList& OutDrawElements, int32 LayerId, const FWidgetStyle& InWidgetStyle, const FAudioSpectrumPlotScaleInfo& ScaleInfo) const;
+	int32 DrawGrid(const FGeometry& AllottedGeometry, FSlateWindowElementList& OutDrawElements, int32 LayerId, const FWidgetStyle& InWidgetStyle, const FAudioSpectrumPlotScaleInfo& ScaleInfo) const;
 	void GetGridLineSoundLevels(TArray<float>& GridLineSoundLevels) const;
 	void GetGridLineFrequencies(TArray<float>& AllGridLineFrequencies, TArray<float>& MajorGridLineFrequencies) const;
 
 	int32 DrawPowerSpectrum(const FGeometry& AllottedGeometry, FSlateWindowElementList& OutDrawElements, int32 LayerId, const FWidgetStyle& InWidgetStyle, const FAudioSpectrumPlotScaleInfo& ScaleInfo) const;
 	FAudioPowerSpectrumData GetPowerSpectrum() const;
 
+	int32 DrawCrosshairAndAxisLabels(const FGeometry& AllottedGeometry, FSlateWindowElementList& OutDrawElements, int32 LayerId, const FWidgetStyle& InWidgetStyle, const FAudioSpectrumPlotScaleInfo& ScaleInfo, TConstArrayView<FVector2f> LinePoints) const;
+
 	// This is a function to reduce the given array of data points to a possibly shorter array of points that will form the line to be plotted.
 	// Where multiple data points map to the same frequency axis pixel bucket, the given 'cost function' will be used to select the best data point (the data point with the lowest 'cost').
 	static void GetSpectrumLinePoints(TArray<FVector2f>& OutLinePoints, TConstArrayView<FVector2f> DataPoints, const FAudioSpectrumPlotScaleInfo& ScaleInfo, TFunctionRef<float(const FVector2f& DataPoint)> CostFunction);
@@ -215,6 +226,7 @@ private:
 	FLinearColor GetBackgroundColor(const FWidgetStyle& InWidgetStyle) const;
 	FLinearColor GetGridColor(const FWidgetStyle& InWidgetStyle) const;
 	FLinearColor GetAxisLabelColor(const FWidgetStyle& InWidgetStyle) const;
+	FLinearColor GetCrosshairColor(const FWidgetStyle& InWidgetStyle) const;
 	FLinearColor GetSpectrumColor(const FWidgetStyle& InWidgetStyle) const;	
 
 	static float GetTiltExponentValue(const EAudioSpectrumPlotTilt InTilt);
@@ -224,6 +236,8 @@ private:
 	void BuildFrequencyAxisScaleSubMenu(FMenuBuilder& SubMenu);
 	void BuildFrequencyAxisPixelBucketModeSubMenu(FMenuBuilder& SubMenu);
 
+	static const float ClampMinSoundLevel;
+
 	static FName ContextMenuExtensionHook;
 	TSharedPtr<FExtender> ContextMenuExtender;
 
@@ -234,6 +248,8 @@ private:
 	TAttribute<float> ViewMaxSoundLevel;
 	TAttribute<float> TiltExponent;
 	TAttribute<float> TiltPivotFrequency;
+	TAttribute<TOptional<float>> SelectedFrequency;
+	TAttribute<bool> bDisplayCrosshair;
 	TAttribute<bool> bDisplayFrequencyAxisLabels;
 	TAttribute<bool> bDisplaySoundLevelAxisLabels;
 	TAttribute<bool> bDisplayFrequencyGridLines;
@@ -243,6 +259,7 @@ private:
 	TAttribute<FSlateColor> BackgroundColor;
 	TAttribute<FSlateColor> GridColor;
 	TAttribute<FSlateColor> AxisLabelColor;
+	TAttribute<FSlateColor> CrosshairColor;
 	TAttribute<FSlateColor> SpectrumColor;
 	TAttribute<bool> bAllowContextMenu;
 	FOnContextMenuOpening OnContextMenuOpening;
