Metasound editor spectrum analyzer now persists its settings within the metasound editor settings.

diff --git a/Engine/Plugins/Runtime/AudioWidgets/Source/AudioWidgets/Private/AudioSpectrumAnalyzer.cpp b/Engine/Plugins/Runtime/AudioWidgets/Source/AudioWidgets/Private/AudioSpectrumAnalyzer.cpp
--- a/Engine/Plugins/Runtime/AudioWidgets/Source/AudioWidgets/Private/AudioSpectrumAnalyzer.cpp
+++ b/Engine/Plugins/Runtime/AudioWidgets/Source/AudioWidgets/Private/AudioSpectrumAnalyzer.cpp
@@ -43,18 +43,34 @@ namespace AudioWidgets
 		}
 	}
 
-	FAudioSpectrumAnalyzer::FAudioSpectrumAnalyzer(int32 InNumChannels, Audio::FDeviceId InAudioDeviceId, TObjectPtr<UAudioBus> InExternalAudioBus)
+	FAudioSpectrumAnalyzer::FAudioSpectrumAnalyzer(const FAudioSpectrumAnalyzerParams& Params)
 		: SpectrumAnalysisSettings(NewObject<USynesthesiaSpectrumAnalysisSettings>())
 		, ConstantQSettings(NewObject<UConstantQSettings>())
 		, Widget(SNew(SAudioSpectrumPlot)
 			.Clipping(EWidgetClipping::ClipToBounds)
+			.TiltExponent(Params.TiltExponent)
 			.DisplayCrosshair(true)
-			.DisplayFrequencyAxisLabels(false)
-			.DisplaySoundLevelAxisLabels(false)
+			.DisplayFrequencyAxisLabels(Params.bDisplayFrequencyAxisLabels)
+			.DisplaySoundLevelAxisLabels(Params.bDisplaySoundLevelAxisLabels)
+			.FrequencyAxisScale(Params.FrequencyAxisScale)
+			.FrequencyAxisPixelBucketMode(Params.FrequencyAxisPixelBucketMode)
+			.OnTiltSpectrumMenuEntryClicked(Params.OnTiltSpectrumMenuEntryClicked)
+			.OnFrequencyAxisPixelBucketModeMenuEntryClicked(Params.OnFrequencyAxisPixelBucketModeMenuEntryClicked)
+			.OnFrequencyAxisScaleMenuEntryClicked(Params.OnFrequencyAxisScaleMenuEntryClicked)
+			.OnDisplayFrequencyAxisLabelsButtonToggled(Params.OnDisplayFrequencyAxisLabelsButtonToggled)
+			.OnDisplaySoundLevelAxisLabelsButtonToggled(Params.OnDisplaySoundLevelAxisLabelsButtonToggled)
 			.OnGetAudioSpectrumData_Raw(this, &FAudioSpectrumAnalyzer::GetAudioSpectrumData))
+		, Ballistics(Params.Ballistics)
+		, AnalyzerType(Params.AnalyzerType)
+		, FFTAnalyzerFFTSize(Params.FFTAnalyzerFFTSize)
+		, CQTAnalyzerFFTSize(Params.CQTAnalyzerFFTSize)
+		, OnBallisticsMenuEntryClicked(Params.OnBallisticsMenuEntryClicked)
+		, OnAnalyzerTypeMenuEntryClicked(Params.OnAnalyzerTypeMenuEntryClicked)
+		, OnFFTAnalyzerFFTSizeMenuEntryClicked(Params.OnFFTAnalyzerFFTSizeMenuEntryClicked)
+		, OnCQTAnalyzerFFTSizeMenuEntryClicked(Params.OnCQTAnalyzerFFTSizeMenuEntryClicked)
 	{
 		SpectrumAnalysisSettings->SpectrumType = EAudioSpectrumType::PowerSpectrum;
-		SpectrumAnalysisSettings->FFTSize = EFFTSize::Max;
+		SpectrumAnalysisSettings->FFTSize = FFTAnalyzerFFTSize.Get();
 		SpectrumAnalysisSettings->WindowType = EFFTWindowType::Blackman;
 		SpectrumAnalysisSettings->bDownmixToMono = true;
 
@@ -62,14 +78,20 @@ namespace AudioWidgets
 		ConstantQSettings->NumBandsPerOctave = 6.0f;
 		ConstantQSettings->NumBands = 61;
 		ConstantQSettings->StartingFrequencyHz = 20000.0f * FMath::Pow(0.5f, (ConstantQSettings->NumBands - 1) / ConstantQSettings->NumBandsPerOctave);
-		ConstantQSettings->FFTSize = EConstantQFFTSizeEnum::XXLarge;
+		ConstantQSettings->FFTSize = CQTAnalyzerFFTSize.Get();
 		ConstantQSettings->WindowType = EFFTWindowType::Blackman;
 		ConstantQSettings->bDownmixToMono = true;
 		ConstantQSettings->BandWidthStretch = 2.0f;
 
 		ContextMenuExtension = Widget->AddContextMenuExtension(EExtensionHook::Before, nullptr, FMenuExtensionDelegate::CreateRaw(this, &FAudioSpectrumAnalyzer::ExtendSpectrumPlotContextMenu));
 
-		Init(InNumChannels, InAudioDeviceId, InExternalAudioBus);
+		Init(Params.NumChannels, Params.AudioDeviceId, Params.ExternalAudioBus);
+	}
+
+	FAudioSpectrumAnalyzer::FAudioSpectrumAnalyzer(int32 InNumChannels, Audio::FDeviceId InAudioDeviceId, TObjectPtr<UAudioBus> InExternalAudioBus)
+		: FAudioSpectrumAnalyzer({ .NumChannels = InNumChannels, .AudioDeviceId = InAudioDeviceId, .ExternalAudioBus = InExternalAudioBus })
+	{
+		// Delegating constructor
 	}
 
 	FAudioSpectrumAnalyzer::~FAudioSpectrumAnalyzer()
@@ -108,12 +130,14 @@ namespace AudioWidgets
 		CreateSynesthesiaSpectrumAnalyzer();
 		CreateConstantQAnalyzer();
 
-		StartAnalyzing();
+		StartAnalyzing(AnalyzerType.Get());
 	}
 
-	void FAudioSpectrumAnalyzer::StartAnalyzing()
+	void FAudioSpectrumAnalyzer::StartAnalyzing(const EAudioSpectrumAnalyzerType InAnalyzerType)
 	{
-		switch (AnalyzerType)
+		ensure(!ActiveAnalyzerType.IsSet());
+
+		switch (InAnalyzerType)
 		{
 		case EAudioSpectrumAnalyzerType::FFT:
 			SpectrumAnalyzer->StartAnalyzing(AudioDeviceId, AudioBus.Get());
@@ -124,11 +148,15 @@ namespace AudioWidgets
 		default:
 			break;
 		}
+
+		ActiveAnalyzerType = InAnalyzerType;
 	}
 
 	void FAudioSpectrumAnalyzer::StopAnalyzing()
 	{
-		switch (AnalyzerType)
+		ensure(ActiveAnalyzerType.IsSet());
+
+		switch (ActiveAnalyzerType.GetValue())
 		{
 		case EAudioSpectrumAnalyzerType::FFT:
 			SpectrumAnalyzer->StopAnalyzing();
@@ -139,11 +167,13 @@ namespace AudioWidgets
 		default:
 			break;
 		}
+
+		ActiveAnalyzerType.Reset();
 	}
 
 	void FAudioSpectrumAnalyzer::OnSpectrumResults(USynesthesiaSpectrumAnalyzer* InSpectrumAnalyzer, int32 ChannelIndex, const TArray<FSynesthesiaSpectrumResults>& InSpectrumResultsArray)
 	{
-		if (AnalyzerType == EAudioSpectrumAnalyzerType::FFT && InSpectrumAnalyzer == SpectrumAnalyzer.Get())
+		if (ActiveAnalyzerType == EAudioSpectrumAnalyzerType::FFT && InSpectrumAnalyzer == SpectrumAnalyzer.Get())
 		{
 			for (const FSynesthesiaSpectrumResults& SpectrumResults : InSpectrumResultsArray)
 			{
@@ -186,7 +216,7 @@ namespace AudioWidgets
 
 	void FAudioSpectrumAnalyzer::OnConstantQResults(UConstantQAnalyzer* InSpectrumAnalyzer, int32 ChannelIndex, const TArray<FConstantQResults>& InSpectrumResultsArray)
 	{
-		if (AnalyzerType == EAudioSpectrumAnalyzerType::CQT && InSpectrumAnalyzer == ConstantQAnalyzer.Get())
+		if (ActiveAnalyzerType == EAudioSpectrumAnalyzerType::CQT && InSpectrumAnalyzer == ConstantQAnalyzer.Get())
 		{
 			for (const FConstantQResults& SpectrumResults : InSpectrumResultsArray)
 			{
@@ -221,6 +251,7 @@ namespace AudioWidgets
 	{
 		// Calculate AR smoother coefficients:
 		const float DeltaT = (TimeStamp - PrevTimeStamp.GetValue());
+		const bool bIsAnalogAttackRelease = (Ballistics.Get() == EAudioSpectrumAnalyzerBallistics::Analog);
 		Audio::FAttackRelease AttackRelease(1.0f / DeltaT, AttackTimeMsec, ReleaseTimeMsec, bIsAnalogAttackRelease);
 
 		// Apply AR smoothing for each frequency:
@@ -238,7 +269,7 @@ namespace AudioWidgets
 	{
 		if (SpectrumAnalyzer.IsValid() && SpectrumAnalyzer->IsValidLowLevel())
 		{
-			if (AnalyzerType == EAudioSpectrumAnalyzerType::FFT)
+			if (ActiveAnalyzerType == EAudioSpectrumAnalyzerType::FFT)
 			{
 				SpectrumAnalyzer->StopAnalyzing();
 			}
@@ -248,7 +279,7 @@ namespace AudioWidgets
 
 		if (ConstantQAnalyzer.IsValid() && ConstantQAnalyzer->IsValidLowLevel())
 		{
-			if (AnalyzerType == EAudioSpectrumAnalyzerType::CQT)
+			if (ActiveAnalyzerType == EAudioSpectrumAnalyzerType::CQT)
 			{
 				ConstantQAnalyzer->StopAnalyzing();
 			}
@@ -256,6 +287,7 @@ namespace AudioWidgets
 			ReleaseConstantQAnalyzer();
 		}
 
+		ActiveAnalyzerType.Reset();
 		PrevTimeStamp.Reset();
 		CenterFrequencies.Empty();
 		ARSmoothedSquaredMagnitudes.Empty();
@@ -264,8 +296,12 @@ namespace AudioWidgets
 		bUseExternalAudioBus = false;
 	}
 
-	FAudioPowerSpectrumData FAudioSpectrumAnalyzer::GetAudioSpectrumData() const
+	FAudioPowerSpectrumData FAudioSpectrumAnalyzer::GetAudioSpectrumData()
 	{
+		// The SAudioSpectrumPlot regularly polls us for audio spectrum data.
+		// We can update the analyzer settings here:
+		UpdateAnalyzerSettings();
+
 		check(CenterFrequencies.Num() == ARSmoothedSquaredMagnitudes.Num());
 		return FAudioPowerSpectrumData{ CenterFrequencies, ARSmoothedSquaredMagnitudes };
 	}
@@ -290,28 +326,36 @@ namespace AudioWidgets
 
 	void FAudioSpectrumAnalyzer::BuildBallisticsSubMenu(FMenuBuilder& SubMenu)
 	{
-		SubMenu.AddMenuEntry(
-			LOCTEXT("Analog", "Analog"),
-			FText(),
-			FSlateIcon(),
-			FUIAction(
-				FExecuteAction::CreateSPLambda(this, [this]() { bIsAnalogAttackRelease = true; }),
-				FCanExecuteAction(),
-				FIsActionChecked::CreateSPLambda(this, [this]() { return bIsAnalogAttackRelease; })
-			),
-			NAME_None,
-			EUserInterfaceActionType::ToggleButton);
-		SubMenu.AddMenuEntry(
-			LOCTEXT("Digital", "Digital"),
-			FText(),
-			FSlateIcon(),
-			FUIAction(
-				FExecuteAction::CreateSPLambda(this, [this]() { bIsAnalogAttackRelease = false; }),
-				FCanExecuteAction(),
-				FIsActionChecked::CreateSPLambda(this, [this]() { return !bIsAnalogAttackRelease; })
-			),
-			NAME_None,
-			EUserInterfaceActionType::ToggleButton);
+		const UEnum* EnumClass = StaticEnum<EAudioSpectrumAnalyzerBallistics>();
+		const int32 NumEnumValues = EnumClass->NumEnums() - 1; // Exclude 'MAX' enum value.
+		for (int32 Index = 0; Index < NumEnumValues; Index++)
+		{
+			const auto EnumValue = static_cast<EAudioSpectrumAnalyzerBallistics>(EnumClass->GetValueByIndex(Index));
+
+			SubMenu.AddMenuEntry(
+				EnumClass->GetDisplayNameTextByIndex(Index),
+#if WITH_EDITOR
+				EnumClass->GetToolTipTextByIndex(Index),
+#else
+				FText(),
+#endif
+				FSlateIcon(),
+				FUIAction(
+					FExecuteAction::CreateSPLambda(this, [this, EnumValue]()
+						{
+							if (!Ballistics.IsBound())
+							{
+								Ballistics = EnumValue;
+							}
+
+							OnBallisticsMenuEntryClicked.ExecuteIfBound(EnumValue);
+						}),
+					FCanExecuteAction(),
+					FIsActionChecked::CreateSPLambda(this, [this, EnumValue]() { return (Ballistics.Get() == EnumValue); })
+				),
+				NAME_None,
+				EUserInterfaceActionType::ToggleButton);
+		}
 	}
 
 	void FAudioSpectrumAnalyzer::BuildAnalyzerTypeSubMenu(FMenuBuilder& SubMenu)
@@ -331,9 +375,17 @@ namespace AudioWidgets
 #endif
 				FSlateIcon(),
 				FUIAction(
-					FExecuteAction::CreateSP(this, &FAudioSpectrumAnalyzer::SetAnalyzerType, EnumValue),
+					FExecuteAction::CreateSPLambda(this, [this, EnumValue]()
+						{
+							if (!AnalyzerType.IsBound())
+							{
+								AnalyzerType = EnumValue;
+							}
+
+							OnAnalyzerTypeMenuEntryClicked.ExecuteIfBound(EnumValue);
+						}),
 					FCanExecuteAction(),
-					FIsActionChecked::CreateSPLambda(this, [this, EnumValue]() { return (AnalyzerType == EnumValue); })
+					FIsActionChecked::CreateSPLambda(this, [this, EnumValue]() { return (AnalyzerType.Get() == EnumValue); })
 				),
 				NAME_None,
 				EUserInterfaceActionType::ToggleButton);
@@ -344,7 +396,7 @@ namespace AudioWidgets
 	{
 		// There is a different FFTSize enum depending on the analyzer type.
 
-		if (AnalyzerType == EAudioSpectrumAnalyzerType::FFT)
+		if (AnalyzerType.Get() == EAudioSpectrumAnalyzerType::FFT)
 		{
 			const UEnum* EnumClass = StaticEnum<EFFTSize>();
 			const int32 NumEnumValues = EnumClass->NumEnums() - 1; // Exclude 'MAX' enum value.
@@ -366,15 +418,23 @@ namespace AudioWidgets
 #endif
 					FSlateIcon(),
 					FUIAction(
-						FExecuteAction::CreateSP(this, &FAudioSpectrumAnalyzer::SetSynesthesiaSpectrumAnalyzerFFTSize, EnumValue),
+						FExecuteAction::CreateSPLambda(this, [this, EnumValue]()
+							{
+								if (!FFTAnalyzerFFTSize.IsBound())
+								{
+									FFTAnalyzerFFTSize = EnumValue;
+								}
+
+								OnFFTAnalyzerFFTSizeMenuEntryClicked.ExecuteIfBound(EnumValue);
+							}),
 						FCanExecuteAction(),
-						FIsActionChecked::CreateSPLambda(this, [this, EnumValue]() { return (SpectrumAnalysisSettings->FFTSize == EnumValue); })
+						FIsActionChecked::CreateSPLambda(this, [this, EnumValue]() { return (FFTAnalyzerFFTSize.Get() == EnumValue); })
 					),
 					NAME_None,
 					EUserInterfaceActionType::ToggleButton);
 			}
 		}
-		else if (AnalyzerType == EAudioSpectrumAnalyzerType::CQT)
+		else if (AnalyzerType.Get() == EAudioSpectrumAnalyzerType::CQT)
 		{
 			const UEnum* EnumClass = StaticEnum<EConstantQFFTSizeEnum>();
 			const int32 NumEnumValues = EnumClass->NumEnums() - 1; // Exclude 'MAX' enum value.
@@ -391,9 +451,17 @@ namespace AudioWidgets
 #endif
 					FSlateIcon(),
 					FUIAction(
-						FExecuteAction::CreateSP(this, &FAudioSpectrumAnalyzer::SetConstantQAnalyzerFFTSize, EnumValue),
+						FExecuteAction::CreateSPLambda(this, [this, EnumValue]()
+							{
+								if (!CQTAnalyzerFFTSize.IsBound())
+								{
+									CQTAnalyzerFFTSize = EnumValue;
+								}
+
+								OnCQTAnalyzerFFTSizeMenuEntryClicked.ExecuteIfBound(EnumValue);
+							}),
 						FCanExecuteAction(),
-						FIsActionChecked::CreateSPLambda(this, [this, EnumValue]() { return (ConstantQSettings->FFTSize == EnumValue); })
+						FIsActionChecked::CreateSPLambda(this, [this, EnumValue]() { return (CQTAnalyzerFFTSize.Get() == EnumValue); })
 					),
 					NAME_None,
 					EUserInterfaceActionType::ToggleButton);
@@ -401,47 +469,42 @@ namespace AudioWidgets
 		}
 	}
 
-	void FAudioSpectrumAnalyzer::SetAnalyzerType(const EAudioSpectrumAnalyzerType InAnalyzerType)
-	{
-		if (AnalyzerType != InAnalyzerType)
-		{
-			StopAnalyzing();
-
-			AnalyzerType = InAnalyzerType;
-
-			PrevTimeStamp.Reset();
-			CenterFrequencies.Reset();
-			ARSmoothedSquaredMagnitudes.Reset();
-
-			StartAnalyzing();
-		}
-	}
-
-	void FAudioSpectrumAnalyzer::SetSynesthesiaSpectrumAnalyzerFFTSize(const EFFTSize FFTSize)
+	void FAudioSpectrumAnalyzer::UpdateAnalyzerSettings()
 	{
-		if (SpectrumAnalysisSettings->FFTSize != FFTSize)
+		const EAudioSpectrumAnalyzerType RequiredAnalyzerType = AnalyzerType.Get();
+		const EFFTSize FFTAnalyzerRequiredFFTSize = FFTAnalyzerFFTSize.Get();
+		const EConstantQFFTSizeEnum CQTAnalyzerRequiredFFTSize = CQTAnalyzerFFTSize.Get();
+
+		const bool bRequiredAnalyzerTypeChanged = (ActiveAnalyzerType != RequiredAnalyzerType);
+		const bool bFFTAnalyzerRequiredFFTSizeChanged = (SpectrumAnalysisSettings->FFTSize != FFTAnalyzerRequiredFFTSize);
+		const bool bCQTAnalyzerRequiredFFTSizeChanged = (ConstantQSettings->FFTSize != CQTAnalyzerRequiredFFTSize);
+		if (bRequiredAnalyzerTypeChanged || bFFTAnalyzerRequiredFFTSizeChanged || bCQTAnalyzerRequiredFFTSizeChanged)
 		{
 			StopAnalyzing();
-			ReleaseSynesthesiaSpectrumAnalyzer();
-
-			SpectrumAnalysisSettings->FFTSize = FFTSize;
 
-			CreateSynesthesiaSpectrumAnalyzer();
-			StartAnalyzing();
-		}
-	}
+			if (bRequiredAnalyzerTypeChanged)
+			{
+				// There will be different center frequencies when the analyzer type changes:
+				PrevTimeStamp.Reset();
+				CenterFrequencies.Reset();
+				ARSmoothedSquaredMagnitudes.Reset();
+			}
 
-	void FAudioSpectrumAnalyzer::SetConstantQAnalyzerFFTSize(const EConstantQFFTSizeEnum FFTSize)
-	{
-		if (ConstantQSettings->FFTSize != FFTSize)
-		{
-			StopAnalyzing();
-			ReleaseConstantQAnalyzer();
+			if (bFFTAnalyzerRequiredFFTSizeChanged)
+			{
+				ReleaseSynesthesiaSpectrumAnalyzer();
+				SpectrumAnalysisSettings->FFTSize = FFTAnalyzerRequiredFFTSize;
+				CreateSynesthesiaSpectrumAnalyzer();
+			}
 
-			ConstantQSettings->FFTSize = FFTSize;
+			if (bCQTAnalyzerRequiredFFTSizeChanged)
+			{
+				ReleaseConstantQAnalyzer();
+				ConstantQSettings->FFTSize = CQTAnalyzerRequiredFFTSize;
+				CreateConstantQAnalyzer();
+			}
 
-			CreateConstantQAnalyzer();
-			StartAnalyzing();
+			StartAnalyzing(RequiredAnalyzerType);
 		}
 	}
 
diff --git a/Engine/Plugins/Runtime/AudioWidgets/Source/AudioWidgets/Private/SAudioSpectrumPlot.cpp b/Engine/Plugins/Runtime/AudioWidgets/Source/AudioWidgets/Private/SAudioSpectrumPlot.cpp
--- a/Engine/Plugins/Runtime/AudioWidgets/Source/AudioWidgets/Private/SAudioSpectrumPlot.cpp
+++ b/Engine/Plugins/Runtime/AudioWidgets/Source/AudioWidgets/Private/SAudioSpectrumPlot.cpp
@@ -435,6 +435,11 @@ void SAudioSpectrumPlot::Construct(const FArguments& InArgs)
 	SpectrumColor = InArgs._SpectrumColor;
 	bAllowContextMenu = InArgs._AllowContextMenu;
 	OnContextMenuOpening = InArgs._OnContextMenuOpening;
+	OnTiltSpectrumMenuEntryClicked = InArgs._OnTiltSpectrumMenuEntryClicked;
+	OnFrequencyAxisPixelBucketModeMenuEntryClicked = InArgs._OnFrequencyAxisPixelBucketModeMenuEntryClicked;
+	OnFrequencyAxisScaleMenuEntryClicked = InArgs._OnFrequencyAxisScaleMenuEntryClicked;
+	OnDisplayFrequencyAxisLabelsButtonToggled = InArgs._OnDisplayFrequencyAxisLabelsButtonToggled;
+	OnDisplaySoundLevelAxisLabelsButtonToggled = InArgs._OnDisplaySoundLevelAxisLabelsButtonToggled;
 	OnGetAudioSpectrumData = InArgs._OnGetAudioSpectrumData;
 }
 
@@ -925,7 +930,7 @@ TSharedRef<SWidget> SAudioSpectrumPlot::BuildDefaultContextMenu()
 
 	MenuBuilder.BeginSection(ContextMenuExtensionHook, LOCTEXT("DisplayOptions", "Display Options"));
 
-	if (!TiltExponent.IsBound())
+	if (OnTiltSpectrumMenuEntryClicked.IsBound() || !TiltExponent.IsBound())
 	{
 		MenuBuilder.AddSubMenu(
 			LOCTEXT("TiltSpectrum", "Tilt Spectrum"),
@@ -933,7 +938,7 @@ TSharedRef<SWidget> SAudioSpectrumPlot::BuildDefaultContextMenu()
 			FNewMenuDelegate::CreateSP(this, &SAudioSpectrumPlot::BuildTiltSpectrumSubMenu));
 	}
 
-	if (!FrequencyAxisPixelBucketMode.IsBound())
+	if (OnFrequencyAxisPixelBucketModeMenuEntryClicked.IsBound() || !FrequencyAxisPixelBucketMode.IsBound())
 	{
 		MenuBuilder.AddSubMenu(
 			LOCTEXT("FrequencyAxisPixelBucketMode", "Pixel Plot Mode"),
@@ -941,7 +946,7 @@ TSharedRef<SWidget> SAudioSpectrumPlot::BuildDefaultContextMenu()
 			FNewMenuDelegate::CreateSP(this, &SAudioSpectrumPlot::BuildFrequencyAxisPixelBucketModeSubMenu));
 	}
 
-	if (!FrequencyAxisScale.IsBound())
+	if (OnFrequencyAxisScaleMenuEntryClicked.IsBound() || !FrequencyAxisScale.IsBound())
 	{
 		MenuBuilder.AddSubMenu(
 			LOCTEXT("FrequencyAxisScale", "Frequency Scale"),
@@ -949,14 +954,22 @@ TSharedRef<SWidget> SAudioSpectrumPlot::BuildDefaultContextMenu()
 			FNewMenuDelegate::CreateSP(this, &SAudioSpectrumPlot::BuildFrequencyAxisScaleSubMenu));
 	}
 
-	if (!bDisplayFrequencyAxisLabels.IsBound())
+	if (OnDisplayFrequencyAxisLabelsButtonToggled.IsBound() || !bDisplayFrequencyAxisLabels.IsBound())
 	{
 		MenuBuilder.AddMenuEntry(
 			LOCTEXT("DisplayFrequencyAxisLabels", "Display Frequency Axis Labels"),
 			FText(),
 			FSlateIcon(),
 			FUIAction(
-				FExecuteAction::CreateSPLambda(this, [this]() { bDisplayFrequencyAxisLabels = !bDisplayFrequencyAxisLabels.Get(); }),
+				FExecuteAction::CreateSPLambda(this, [this]()
+					{
+						if (!bDisplayFrequencyAxisLabels.IsBound())
+						{
+							bDisplayFrequencyAxisLabels = !bDisplayFrequencyAxisLabels.Get();
+						}
+
+						OnDisplayFrequencyAxisLabelsButtonToggled.ExecuteIfBound();
+					}),
 				FCanExecuteAction(),
 				FIsActionChecked::CreateSPLambda(this, [this]() { return bDisplayFrequencyAxisLabels.Get(); })
 			),
@@ -964,14 +977,22 @@ TSharedRef<SWidget> SAudioSpectrumPlot::BuildDefaultContextMenu()
 			EUserInterfaceActionType::ToggleButton);
 	}
 
-	if (!bDisplaySoundLevelAxisLabels.IsBound())
+	if (OnDisplaySoundLevelAxisLabelsButtonToggled.IsBound() || !bDisplaySoundLevelAxisLabels.IsBound())
 	{
 		MenuBuilder.AddMenuEntry(
 			LOCTEXT("DisplaySoundLevelAxisLabels", "Display Sound Level Axis Labels"),
 			FText(),
 			FSlateIcon(),
 			FUIAction(
-				FExecuteAction::CreateSPLambda(this, [this]() { bDisplaySoundLevelAxisLabels = !bDisplaySoundLevelAxisLabels.Get(); }),
+				FExecuteAction::CreateSPLambda(this, [this]()
+					{
+						if (!bDisplaySoundLevelAxisLabels.IsBound())
+						{
+							bDisplaySoundLevelAxisLabels = !bDisplaySoundLevelAxisLabels.Get();
+						}
+
+						OnDisplaySoundLevelAxisLabelsButtonToggled.ExecuteIfBound();
+					}),
 				FCanExecuteAction(),
 				FIsActionChecked::CreateSPLambda(this, [this]() { return bDisplaySoundLevelAxisLabels.Get(); })
 			),
@@ -1002,7 +1023,15 @@ void SAudioSpectrumPlot::BuildTiltSpectrumSubMenu(FMenuBuilder& SubMenu)
 #endif
 			FSlateIcon(),
 			FUIAction(
-				FExecuteAction::CreateSPLambda(this, [this, TiltExponentValue]() { TiltExponent = TiltExponentValue; }),
+				FExecuteAction::CreateSPLambda(this, [this, EnumValue, TiltExponentValue]()
+					{
+						if (!TiltExponent.IsBound())
+						{
+							TiltExponent = TiltExponentValue;
+						}
+
+						OnTiltSpectrumMenuEntryClicked.ExecuteIfBound(EnumValue);
+					}),
 				FCanExecuteAction(),
 				FIsActionChecked::CreateSPLambda(this, [this, TiltExponentValue]() { return (TiltExponent.Get() == TiltExponentValue); })
 			),
@@ -1028,7 +1057,15 @@ void SAudioSpectrumPlot::BuildFrequencyAxisScaleSubMenu(FMenuBuilder& SubMenu)
 #endif
 			FSlateIcon(),
 			FUIAction(
-				FExecuteAction::CreateSP(this, &SAudioSpectrumPlot::SetFrequencyAxisScale, EnumValue),
+				FExecuteAction::CreateSPLambda(this, [this, EnumValue]()
+					{
+						if (!FrequencyAxisScale.IsBound())
+						{
+							FrequencyAxisScale = EnumValue;
+						}
+
+						OnFrequencyAxisScaleMenuEntryClicked.ExecuteIfBound(EnumValue);
+					}),
 				FCanExecuteAction(),
 				FIsActionChecked::CreateSPLambda(this, [this, EnumValue]() { return (FrequencyAxisScale.Get() == EnumValue); })
 			),
@@ -1054,7 +1091,15 @@ void SAudioSpectrumPlot::BuildFrequencyAxisPixelBucketModeSubMenu(FMenuBuilder&
 #endif
 			FSlateIcon(),
 			FUIAction(
-				FExecuteAction::CreateSP(this, &SAudioSpectrumPlot::SetFrequencyAxisPixelBucketMode, EnumValue),
+				FExecuteAction::CreateSPLambda(this, [this, EnumValue]()
+					{
+						if (!FrequencyAxisPixelBucketMode.IsBound())
+						{
+							FrequencyAxisPixelBucketMode = EnumValue;
+						}
+
+						OnFrequencyAxisPixelBucketModeMenuEntryClicked.ExecuteIfBound(EnumValue);
+					}),
 				FCanExecuteAction(),
 				FIsActionChecked::CreateSPLambda(this, [this, EnumValue]() { return (FrequencyAxisPixelBucketMode.Get() == EnumValue); })
 			),
diff --git a/Engine/Plugins/Runtime/AudioWidgets/Source/AudioWidgets/Public/AudioSpectrumAnalyzer.h b/Engine/Plugins/Runtime/AudioWidgets/Source/AudioWidgets/Public/AudioSpectrumAnalyzer.h
--- a/Engine/Plugins/Runtime/AudioWidgets/Source/AudioWidgets/Public/AudioSpectrumAnalyzer.h
+++ b/Engine/Plugins/Runtime/AudioWidgets/Source/AudioWidgets/Public/AudioSpectrumAnalyzer.h
@@ -10,6 +10,13 @@
 
 class UWorld;
 
+UENUM(BlueprintType)
+enum class EAudioSpectrumAnalyzerBallistics : uint8
+{
+	Analog,
+	Digital,
+};
+
 UENUM(BlueprintType)
 enum class EAudioSpectrumAnalyzerType : uint8
 {
@@ -17,8 +24,43 @@ enum class EAudioSpectrumAnalyzerType : uint8
 	CQT UMETA(ToolTip = "Constant-Q Transform"),
 };
 
+DECLARE_DELEGATE_OneParam(FOnBallisticsMenuEntryClicked, EAudioSpectrumAnalyzerBallistics);
+DECLARE_DELEGATE_OneParam(FOnAnalyzerTypeMenuEntryClicked, EAudioSpectrumAnalyzerType);
+DECLARE_DELEGATE_OneParam(FOnFFTAnalyzerFFTSizeMenuEntryClicked, EFFTSize);
+DECLARE_DELEGATE_OneParam(FOnCQTAnalyzerFFTSizeMenuEntryClicked, EConstantQFFTSizeEnum);
+
 namespace AudioWidgets
 {
+	/**
+	 * Constructor parameters for the analyzer.
+	 */
+	struct FAudioSpectrumAnalyzerParams
+	{
+		int32 NumChannels = 1;
+		Audio::FDeviceId AudioDeviceId = INDEX_NONE;
+		TObjectPtr<UAudioBus> ExternalAudioBus = nullptr;
+
+		TAttribute<EAudioSpectrumAnalyzerBallistics> Ballistics = EAudioSpectrumAnalyzerBallistics::Digital;
+		TAttribute<EAudioSpectrumAnalyzerType> AnalyzerType = EAudioSpectrumAnalyzerType::CQT;
+		TAttribute<EFFTSize> FFTAnalyzerFFTSize = EFFTSize::Max;
+		TAttribute<EConstantQFFTSizeEnum> CQTAnalyzerFFTSize = EConstantQFFTSizeEnum::XXLarge;
+		TAttribute<float> TiltExponent = 0.0f;
+		TAttribute<EAudioSpectrumPlotFrequencyAxisPixelBucketMode> FrequencyAxisPixelBucketMode = EAudioSpectrumPlotFrequencyAxisPixelBucketMode::Average;
+		TAttribute<EAudioSpectrumPlotFrequencyAxisScale> FrequencyAxisScale = EAudioSpectrumPlotFrequencyAxisScale::Logarithmic;
+		TAttribute<bool> bDisplayFrequencyAxisLabels = false;
+		TAttribute<bool> bDisplaySoundLevelAxisLabels = false;
+		
+		FOnBallisticsMenuEntryClicked OnBallisticsMenuEntryClicked;
+		FOnAnalyzerTypeMenuEntryClicked OnAnalyzerTypeMenuEntryClicked;
+		FOnFFTAnalyzerFFTSizeMenuEntryClicked OnFFTAnalyzerFFTSizeMenuEntryClicked;
+		FOnCQTAnalyzerFFTSizeMenuEntryClicked OnCQTAnalyzerFFTSizeMenuEntryClicked;
+		FOnTiltSpectrumMenuEntryClicked OnTiltSpectrumMenuEntryClicked;
+		FOnFrequencyAxisPixelBucketModeMenuEntryClicked OnFrequencyAxisPixelBucketModeMenuEntryClicked;
+		FOnFrequencyAxisScaleMenuEntryClicked OnFrequencyAxisScaleMenuEntryClicked;
+		FOnDisplayAxisLabelsButtonToggled OnDisplayFrequencyAxisLabelsButtonToggled;
+		FOnDisplayAxisLabelsButtonToggled OnDisplaySoundLevelAxisLabelsButtonToggled;
+	};
+
 	/**
 	 * Owns an analyzer and a corresponding Slate widget for displaying the resulting spectrum.
 	 * Exponential time-smoothing is applied to the spectrum.
@@ -27,6 +69,7 @@ namespace AudioWidgets
 	class AUDIOWIDGETS_API FAudioSpectrumAnalyzer : public TSharedFromThis<FAudioSpectrumAnalyzer>
 	{
 	public:
+		FAudioSpectrumAnalyzer(const FAudioSpectrumAnalyzerParams& Params);
 		FAudioSpectrumAnalyzer(int32 InNumChannels, Audio::FDeviceId InAudioDeviceId, TObjectPtr<UAudioBus> InExternalAudioBus = nullptr);
 		~FAudioSpectrumAnalyzer();
 
@@ -37,23 +80,21 @@ namespace AudioWidgets
 		void Init(int32 InNumChannels, Audio::FDeviceId InAudioDeviceId, TObjectPtr<UAudioBus> InExternalAudioBus = nullptr);
 
 	protected:
-		void StartAnalyzing();
+		void StartAnalyzing(const EAudioSpectrumAnalyzerType InAnalyzerType);
 		void StopAnalyzing();
 
 		void OnSpectrumResults(USynesthesiaSpectrumAnalyzer* InSpectrumAnalyzer, int32 ChannelIndex, const TArray<FSynesthesiaSpectrumResults>& InSpectrumResultsArray);
 		void OnConstantQResults(UConstantQAnalyzer* InSpectrumAnalyzer, int32 ChannelIndex, const TArray<FConstantQResults>& InSpectrumResultsArray);
 		void UpdateARSmoothing(const float TimeStamp, TConstArrayView<float> SquaredMagnitudes);
 
-		FAudioPowerSpectrumData GetAudioSpectrumData() const;
+		FAudioPowerSpectrumData GetAudioSpectrumData();
 
 		void ExtendSpectrumPlotContextMenu(FMenuBuilder& MenuBuilder);
 		void BuildBallisticsSubMenu(FMenuBuilder& SubMenu);
 		void BuildAnalyzerTypeSubMenu(FMenuBuilder& SubMenu);
 		void BuildFFTSizeSubMenu(FMenuBuilder& SubMenu);
 
-		void SetAnalyzerType(const EAudioSpectrumAnalyzerType InAnalyzerType);
-		void SetSynesthesiaSpectrumAnalyzerFFTSize(const EFFTSize FFTSize);
-		void SetConstantQAnalyzerFFTSize(const EConstantQFFTSizeEnum FFTSize);
+		void UpdateAnalyzerSettings();
 
 	private:
 		void CreateSynesthesiaSpectrumAnalyzer();
@@ -92,12 +133,19 @@ namespace AudioWidgets
 		Audio::FDeviceId AudioDeviceId = INDEX_NONE;
 		bool bUseExternalAudioBus = false;
 
-		EAudioSpectrumAnalyzerType AnalyzerType = EAudioSpectrumAnalyzerType::CQT;
-
+		TOptional<EAudioSpectrumAnalyzerType> ActiveAnalyzerType;
 		TOptional<float> PrevTimeStamp;
 		float WindowCompensationPowerGain = 1.0f;
 		float AttackTimeMsec = 300.0f;
 		float ReleaseTimeMsec = 300.0f;
-		bool bIsAnalogAttackRelease = false;
+		TAttribute<EAudioSpectrumAnalyzerBallistics> Ballistics;
+		TAttribute<EAudioSpectrumAnalyzerType> AnalyzerType;
+		TAttribute<EFFTSize> FFTAnalyzerFFTSize;
+		TAttribute<EConstantQFFTSizeEnum> CQTAnalyzerFFTSize;
+
+		FOnBallisticsMenuEntryClicked OnBallisticsMenuEntryClicked;
+		FOnAnalyzerTypeMenuEntryClicked OnAnalyzerTypeMenuEntryClicked;
+		FOnFFTAnalyzerFFTSizeMenuEntryClicked OnFFTAnalyzerFFTSizeMenuEntryClicked;
+		FOnCQTAnalyzerFFTSizeMenuEntryClicked OnCQTAnalyzerFFTSizeMenuEntryClicked;
 	};
 } // namespace AudioWidgets
diff --git a/Engine/Plugins/Runtime/AudioWidgets/Source/AudioWidgets/Public/SAudioSpectrumPlot.h b/Engine/Plugins/Runtime/AudioWidgets/Source/AudioWidgets/Public/SAudioSpectrumPlot.h
--- a/Engine/Plugins/Runtime/AudioWidgets/Source/AudioWidgets/Public/SAudioSpectrumPlot.h
+++ b/Engine/Plugins/Runtime/AudioWidgets/Source/AudioWidgets/Public/SAudioSpectrumPlot.h
@@ -107,6 +107,11 @@ private:
 	const float PixelsPerDecibel;
 };
 
+DECLARE_DELEGATE_OneParam(FOnTiltSpectrumMenuEntryClicked, EAudioSpectrumPlotTilt);
+DECLARE_DELEGATE_OneParam(FOnFrequencyAxisPixelBucketModeMenuEntryClicked, EAudioSpectrumPlotFrequencyAxisPixelBucketMode);
+DECLARE_DELEGATE_OneParam(FOnFrequencyAxisScaleMenuEntryClicked, EAudioSpectrumPlotFrequencyAxisScale);
+DECLARE_DELEGATE(FOnDisplayAxisLabelsButtonToggled);
+
 /**
  * The audio spectrum data to plot.
  */
@@ -168,6 +173,11 @@ public:
 		SLATE_ATTRIBUTE(FSlateColor, SpectrumColor)
 		SLATE_ATTRIBUTE(bool, AllowContextMenu)
 		SLATE_EVENT(FOnContextMenuOpening, OnContextMenuOpening)
+		SLATE_EVENT(FOnTiltSpectrumMenuEntryClicked, OnTiltSpectrumMenuEntryClicked)
+		SLATE_EVENT(FOnFrequencyAxisPixelBucketModeMenuEntryClicked, OnFrequencyAxisPixelBucketModeMenuEntryClicked)
+		SLATE_EVENT(FOnFrequencyAxisScaleMenuEntryClicked, OnFrequencyAxisScaleMenuEntryClicked)
+		SLATE_EVENT(FOnDisplayAxisLabelsButtonToggled, OnDisplayFrequencyAxisLabelsButtonToggled)
+		SLATE_EVENT(FOnDisplayAxisLabelsButtonToggled, OnDisplaySoundLevelAxisLabelsButtonToggled)
 		SLATE_EVENT(FGetAudioSpectrumData, OnGetAudioSpectrumData)
 	SLATE_END_ARGS()
 
@@ -204,6 +214,8 @@ public:
 
 	FAudioSpectrumPlotScaleInfo GetScaleInfo() const;
 
+	static float GetTiltExponentValue(const EAudioSpectrumPlotTilt InTilt);
+
 private:
 	// Begin SWidget overrides.
 	virtual int32 OnPaint(const FPaintArgs& Args, const FGeometry& AllottedGeometry, const FSlateRect& MyCullingRect, FSlateWindowElementList& OutDrawElements, int32 LayerId, const FWidgetStyle& InWidgetStyle, bool bParentEnabled) const override;
@@ -230,8 +242,6 @@ private:
 	FLinearColor GetCrosshairColor(const FWidgetStyle& InWidgetStyle) const;
 	FLinearColor GetSpectrumColor(const FWidgetStyle& InWidgetStyle) const;	
 
-	static float GetTiltExponentValue(const EAudioSpectrumPlotTilt InTilt);
-
 	TSharedRef<SWidget> BuildDefaultContextMenu();
 	void BuildTiltSpectrumSubMenu(FMenuBuilder& SubMenu);
 	void BuildFrequencyAxisScaleSubMenu(FMenuBuilder& SubMenu);
@@ -264,5 +274,10 @@ private:
 	TAttribute<FSlateColor> SpectrumColor;
 	TAttribute<bool> bAllowContextMenu;
 	FOnContextMenuOpening OnContextMenuOpening;
+	FOnTiltSpectrumMenuEntryClicked OnTiltSpectrumMenuEntryClicked;
+	FOnFrequencyAxisPixelBucketModeMenuEntryClicked OnFrequencyAxisPixelBucketModeMenuEntryClicked;
+	FOnFrequencyAxisScaleMenuEntryClicked OnFrequencyAxisScaleMenuEntryClicked;
+	FOnDisplayAxisLabelsButtonToggled OnDisplayFrequencyAxisLabelsButtonToggled;
+	FOnDisplayAxisLabelsButtonToggled OnDisplaySoundLevelAxisLabelsButtonToggled;
 	FGetAudioSpectrumData OnGetAudioSpectrumData;
 };
diff --git a/Engine/Plugins/Runtime/Metasound/Source/MetasoundEditor/Private/MetasoundEditor.cpp b/Engine/Plugins/Runtime/Metasound/Source/MetasoundEditor/Private/MetasoundEditor.cpp
--- a/Engine/Plugins/Runtime/Metasound/Source/MetasoundEditor/Private/MetasoundEditor.cpp
+++ b/Engine/Plugins/Runtime/Metasound/Source/MetasoundEditor/Private/MetasoundEditor.cpp
@@ -1622,7 +1622,97 @@ namespace Metasound
 
 				if (!OutputSpectrumAnalyzer.IsValid())
 				{
-					OutputSpectrumAnalyzer = MakeShared<AudioWidgets::FAudioSpectrumAnalyzer>(MetaSoundSource.NumChannels, AudioDeviceId);
+					AudioWidgets::FAudioSpectrumAnalyzerParams Params;
+					Params.NumChannels = MetaSoundSource.NumChannels;
+					Params.AudioDeviceId = AudioDeviceId;
+
+					Params.Ballistics.BindLambda([]()
+						{
+							return GetDefault<UMetasoundEditorSettings>()->SpectrumAnalyzerSettings.Ballistics;
+						});
+					Params.AnalyzerType.BindLambda([]()
+						{
+							return GetDefault<UMetasoundEditorSettings>()->SpectrumAnalyzerSettings.AnalyzerType;
+						});
+					Params.FFTAnalyzerFFTSize.BindLambda([]()
+						{
+							return GetDefault<UMetasoundEditorSettings>()->SpectrumAnalyzerSettings.FFTAnalyzerFFTSize;
+						});
+					Params.CQTAnalyzerFFTSize.BindLambda([]()
+						{
+							return GetDefault<UMetasoundEditorSettings>()->SpectrumAnalyzerSettings.CQTAnalyzerFFTSize;
+						});
+					Params.TiltExponent.BindLambda([]()
+						{
+							const EAudioSpectrumPlotTilt TiltSpectrum = GetDefault<UMetasoundEditorSettings>()->SpectrumAnalyzerSettings.TiltSpectrum;
+							return SAudioSpectrumPlot::GetTiltExponentValue(TiltSpectrum);
+						});
+					Params.FrequencyAxisPixelBucketMode.BindLambda([]()
+						{
+							return GetDefault<UMetasoundEditorSettings>()->SpectrumAnalyzerSettings.PixelPlotMode;
+						});
+					Params.FrequencyAxisScale.BindLambda([]()
+						{
+							return GetDefault<UMetasoundEditorSettings>()->SpectrumAnalyzerSettings.FrequencyScale;
+						});
+					Params.bDisplayFrequencyAxisLabels.BindLambda([]()
+						{
+							return GetDefault<UMetasoundEditorSettings>()->SpectrumAnalyzerSettings.bDisplayFrequencyAxisLabels;
+						});
+					Params.bDisplaySoundLevelAxisLabels.BindLambda([]()
+						{
+							return GetDefault<UMetasoundEditorSettings>()->SpectrumAnalyzerSettings.bDisplaySoundLevelAxisLabels;
+						});
+
+					Params.OnBallisticsMenuEntryClicked.BindLambda([](EAudioSpectrumAnalyzerBallistics SelectedValue)
+						{
+							GetMutableDefault<UMetasoundEditorSettings>()->SpectrumAnalyzerSettings.Ballistics = SelectedValue;
+							GetMutableDefault<UMetasoundEditorSettings>()->SaveConfig();
+						});
+					Params.OnAnalyzerTypeMenuEntryClicked.BindLambda([](EAudioSpectrumAnalyzerType SelectedValue)
+						{
+							GetMutableDefault<UMetasoundEditorSettings>()->SpectrumAnalyzerSettings.AnalyzerType = SelectedValue;
+							GetMutableDefault<UMetasoundEditorSettings>()->SaveConfig();
+						});
+					Params.OnFFTAnalyzerFFTSizeMenuEntryClicked.BindLambda([](EFFTSize SelectedValue)
+						{
+							GetMutableDefault<UMetasoundEditorSettings>()->SpectrumAnalyzerSettings.FFTAnalyzerFFTSize = SelectedValue;
+							GetMutableDefault<UMetasoundEditorSettings>()->SaveConfig();
+						});
+					Params.OnCQTAnalyzerFFTSizeMenuEntryClicked.BindLambda([](EConstantQFFTSizeEnum SelectedValue)
+						{
+							GetMutableDefault<UMetasoundEditorSettings>()->SpectrumAnalyzerSettings.CQTAnalyzerFFTSize = SelectedValue;
+							GetMutableDefault<UMetasoundEditorSettings>()->SaveConfig();
+						});
+					Params.OnTiltSpectrumMenuEntryClicked.BindLambda([](EAudioSpectrumPlotTilt SelectedValue)
+						{
+							GetMutableDefault<UMetasoundEditorSettings>()->SpectrumAnalyzerSettings.TiltSpectrum = SelectedValue;
+							GetMutableDefault<UMetasoundEditorSettings>()->SaveConfig();
+						});
+					Params.OnFrequencyAxisPixelBucketModeMenuEntryClicked.BindLambda([](EAudioSpectrumPlotFrequencyAxisPixelBucketMode SelectedValue)
+						{
+							GetMutableDefault<UMetasoundEditorSettings>()->SpectrumAnalyzerSettings.PixelPlotMode = SelectedValue;
+							GetMutableDefault<UMetasoundEditorSettings>()->SaveConfig();
+						});
+					Params.OnFrequencyAxisScaleMenuEntryClicked.BindLambda([](EAudioSpectrumPlotFrequencyAxisScale SelectedValue)
+						{
+							GetMutableDefault<UMetasoundEditorSettings>()->SpectrumAnalyzerSettings.FrequencyScale = SelectedValue;
+							GetMutableDefault<UMetasoundEditorSettings>()->SaveConfig();
+						});
+					Params.OnDisplayFrequencyAxisLabelsButtonToggled.BindLambda([]()
+						{
+							FMetasoundEditorSpectrumAnalyzerSettings& SpectrumAnalyzerSettings = GetMutableDefault<UMetasoundEditorSettings>()->SpectrumAnalyzerSettings;
+							SpectrumAnalyzerSettings.bDisplayFrequencyAxisLabels = !SpectrumAnalyzerSettings.bDisplayFrequencyAxisLabels;
+							GetMutableDefault<UMetasoundEditorSettings>()->SaveConfig();
+						});
+					Params.OnDisplaySoundLevelAxisLabelsButtonToggled.BindLambda([]()
+						{
+							FMetasoundEditorSpectrumAnalyzerSettings& SpectrumAnalyzerSettings = GetMutableDefault<UMetasoundEditorSettings>()->SpectrumAnalyzerSettings;
+							SpectrumAnalyzerSettings.bDisplaySoundLevelAxisLabels = !SpectrumAnalyzerSettings.bDisplaySoundLevelAxisLabels;
+							GetMutableDefault<UMetasoundEditorSettings>()->SaveConfig();
+						});
+
+					OutputSpectrumAnalyzer = MakeShared<AudioWidgets::FAudioSpectrumAnalyzer>(Params);
 				}
 				else if (OutputSpectrumAnalyzer->GetAudioBus()->GetNumChannels() != MetaSoundSource.NumChannels)
 				{
diff --git a/Engine/Plugins/Runtime/Metasound/Source/MetasoundEditor/Private/MetasoundEditorSettings.h b/Engine/Plugins/Runtime/Metasound/Source/MetasoundEditor/Private/MetasoundEditorSettings.h
--- a/Engine/Plugins/Runtime/Metasound/Source/MetasoundEditor/Private/MetasoundEditorSettings.h
+++ b/Engine/Plugins/Runtime/Metasound/Source/MetasoundEditor/Private/MetasoundEditorSettings.h
@@ -1,6 +1,7 @@
 // Copyright Epic Games, Inc. All Rights Reserved.
 #pragma once
 
+#include "AudioSpectrumAnalyzer.h"
 #include "Misc/CoreDefines.h"
 #include "UObject/NoExportTypes.h"
 #include "UObject/Object.h"
@@ -95,6 +96,39 @@ struct FMetasoundAnalyzerAnimationSettings
 	float WireScalarMax = 4.5f;
 };
 
+USTRUCT()
+struct FMetasoundEditorSpectrumAnalyzerSettings
+{
+	GENERATED_BODY()
+	
+	UPROPERTY(EditAnywhere, config, Category = SpectrumAnalyzer)
+	EAudioSpectrumAnalyzerBallistics Ballistics = EAudioSpectrumAnalyzerBallistics::Digital;
+
+	UPROPERTY(EditAnywhere, config, Category = SpectrumAnalyzer)
+	EAudioSpectrumAnalyzerType AnalyzerType = EAudioSpectrumAnalyzerType::CQT;
+
+	UPROPERTY(EditAnywhere, config, Category = SpectrumAnalyzer, meta = (DisplayName = "FFT Size (FFT Analyzer)"))
+	EFFTSize FFTAnalyzerFFTSize = EFFTSize::Max;
+
+	UPROPERTY(EditAnywhere, config, Category = SpectrumAnalyzer, meta = (DisplayName = "FFT Size (CQT Analyzer)"))
+	EConstantQFFTSizeEnum CQTAnalyzerFFTSize = EConstantQFFTSizeEnum::XXLarge;
+
+	UPROPERTY(EditAnywhere, config, Category = SpectrumAnalyzer)
+	EAudioSpectrumPlotTilt TiltSpectrum = EAudioSpectrumPlotTilt::NoTilt;
+
+	UPROPERTY(EditAnywhere, config, Category = SpectrumAnalyzer)
+	EAudioSpectrumPlotFrequencyAxisPixelBucketMode PixelPlotMode = EAudioSpectrumPlotFrequencyAxisPixelBucketMode::Average;
+
+	UPROPERTY(EditAnywhere, config, Category = SpectrumAnalyzer)
+	EAudioSpectrumPlotFrequencyAxisScale FrequencyScale = EAudioSpectrumPlotFrequencyAxisScale::Logarithmic;
+
+	UPROPERTY(EditAnywhere, config, Category = SpectrumAnalyzer)
+	bool bDisplayFrequencyAxisLabels = false;
+
+	UPROPERTY(EditAnywhere, config, Category = SpectrumAnalyzer)
+	bool bDisplaySoundLevelAxisLabels = false;
+};
+
 UCLASS(config=EditorPerProjectUserSettings)
 class METASOUNDEDITOR_API UMetasoundEditorSettings : public UObject
 {
@@ -197,6 +231,10 @@ public:
 	UPROPERTY(EditAnywhere, config, Category = NodeTitleColors)
 	FLinearColor VariableNodeTitleColor;
 
+	/** Settings for metasound output spectrum analyzer widget */
+	UPROPERTY(EditAnywhere, config, Category = SpectrumAnalyzer, meta = (ShowOnlyInnerProperties))
+	FMetasoundEditorSpectrumAnalyzerSettings SpectrumAnalyzerSettings;
+
 	/** Widget type to show on input nodes by default */
 	UPROPERTY(EditAnywhere, config, Category = General)
 	EMetasoundMemberDefaultWidget DefaultInputWidgetType = EMetasoundMemberDefaultWidget::RadialSlider;
