Resource load in AudioWidgetsStyle causes warnings in game targets.
#rnx

#virtualized

diff --git a/Engine/Plugins/Runtime/AudioWidgets/Source/AudioWidgets/Private/AudioWidgetsStyle.cpp b/Engine/Plugins/Runtime/AudioWidgets/Source/AudioWidgets/Private/AudioWidgetsStyle.cpp
--- a/Engine/Plugins/Runtime/AudioWidgets/Source/AudioWidgets/Private/AudioWidgetsStyle.cpp
+++ b/Engine/Plugins/Runtime/AudioWidgets/Source/AudioWidgets/Private/AudioWidgetsStyle.cpp
@@ -3,8 +3,10 @@
 #include "AudioWidgetsStyle.h"
 #include "AudioWidgetsSlateTypes.h"
 #include "AudioMaterialSlate/AudioMaterialSlateTypes.h"
+#include "Brushes/SlateBoxBrush.h"
 #include "Brushes/SlateNoResource.h"
 #include "Brushes/SlateRoundedBoxBrush.h"
+#include "Engine/Texture2D.h"
 #include "Misc/Paths.h"
 #include "Styling/StyleColors.h"
 #include "Styling/SlateStyleRegistry.h"
@@ -156,7 +158,7 @@ FAudioWidgetsStyle::FAudioWidgetsStyle()
 }
 void FAudioWidgetsStyle::SetResources()
 {
-	Set(AudioWidgetsStylePrivate::ScrubHandleBrushName, new FSlateBoxBrush(RootToContentDir(TEXT("Resources/ScrubHandleDown_Clamped.png")), FMargin(6.f / 13.f, 3.f / 12.f, 6.f / 13.f, 7.f / 12.f)));
+	Set(AudioWidgetsStylePrivate::ScrubHandleBrushName, new FSlateBoxBrush(LoadObject<UTexture2D>(nullptr, TEXT("/AudioWidgets/Resources/ScrubHandleDown_Clamped.ScrubHandleDown_Clamped")), FMargin(6.f / 13.f, 3.f / 12.f, 6.f / 13.f, 7.f / 12.f)));
 }
 
 FAudioWidgetsStyle::~FAudioWidgetsStyle()
