Create a new AudioDetailsCustomization module inside AudioWidgets to remove audio related details views from editor
#rnx

diff --git a/Engine/Plugins/Runtime/AudioWidgets/AudioWidgets.uplugin b/Engine/Plugins/Runtime/AudioWidgets/AudioWidgets.uplugin
--- a/Engine/Plugins/Runtime/AudioWidgets/AudioWidgets.uplugin
+++ b/Engine/Plugins/Runtime/AudioWidgets/AudioWidgets.uplugin
@@ -19,6 +19,11 @@
 			"Name": "AudioWidgets",
 			"Type": "Runtime",
 			"LoadingPhase": "Default"
+		},
+		{
+			"Name": "AudioWidgetsEditor",
+			"Type": "Editor",
+			"LoadingPhase": "Default"
 		}
 	],
 	"Plugins": [
diff --git a/Engine/Plugins/Runtime/AudioWidgets/Source/AudioWidgetsEditor/AudioWidgetsEditor.build.cs b/Engine/Plugins/Runtime/AudioWidgets/Source/AudioWidgetsEditor/AudioWidgetsEditor.build.cs
new file mode 100644
--- /dev/null
+++ b/Engine/Plugins/Runtime/AudioWidgets/Source/AudioWidgetsEditor/AudioWidgetsEditor.build.cs
@@ -0,0 +1,20 @@
+// Copyright Epic Games, Inc. All Rights Reserved.
+
+using UnrealBuildTool;
+
+public class AudioWidgetsEditor : ModuleRules
+{
+	public AudioWidgetsEditor(ReadOnlyTargetRules Target) : base(Target)
+	{
+		PCHUsage = ModuleRules.PCHUsageMode.UseExplicitOrSharedPCHs;
+
+		PrivateDependencyModuleNames.AddRange
+		(
+			new string[] 
+			{
+				"Core",
+				"PropertyEditor",
+			}
+		);
+	}
+}
\ No newline at end of file
diff --git a/Engine/Plugins/Runtime/AudioWidgets/Source/AudioWidgetsEditor/Private/AudioWidgetsEditorModule.cpp b/Engine/Plugins/Runtime/AudioWidgets/Source/AudioWidgetsEditor/Private/AudioWidgetsEditorModule.cpp
new file mode 100644
--- /dev/null
+++ b/Engine/Plugins/Runtime/AudioWidgets/Source/AudioWidgetsEditor/Private/AudioWidgetsEditorModule.cpp
@@ -0,0 +1,16 @@
+// Copyright Epic Games, Inc. All Rights Reserved.
+
+#include "AudioWidgetsEditorModule.h"
+#include "Modules/ModuleManager.h"
+
+
+void FAudioWidgetsEditorModule::StartupModule()
+{
+}
+
+void FAudioWidgetsEditorModule::ShutdownModule()
+{
+}
+
+	
+IMPLEMENT_MODULE(FAudioWidgetsEditorModule, AudioWidgetsEditor)
\ No newline at end of file
diff --git a/Engine/Plugins/Runtime/AudioWidgets/Source/AudioWidgetsEditor/Public/AudioWidgetsEditorModule.h b/Engine/Plugins/Runtime/AudioWidgets/Source/AudioWidgetsEditor/Public/AudioWidgetsEditorModule.h
new file mode 100644
--- /dev/null
+++ b/Engine/Plugins/Runtime/AudioWidgets/Source/AudioWidgetsEditor/Public/AudioWidgetsEditorModule.h
@@ -0,0 +1,14 @@
+// Copyright Epic Games, Inc. All Rights Reserved.
+
+#pragma once
+
+#include "Modules/ModuleInterface.h"
+
+class FAudioWidgetsEditorModule : public IModuleInterface
+{
+public:
+
+	/** IModuleInterface implementation */
+	virtual void StartupModule() override;
+	virtual void ShutdownModule() override;
+};
\ No newline at end of file
