Move Soundbase details outside of editor

diff --git a/Engine/Plugins/Runtime/AudioWidgets/Source/AudioWidgetsEditor/Private/AudioWidgetsEditorModule.cpp b/Engine/Plugins/Runtime/AudioWidgets/Source/AudioWidgetsEditor/Private/AudioWidgetsEditorModule.cpp
--- a/Engine/Plugins/Runtime/AudioWidgets/Source/AudioWidgetsEditor/Private/AudioWidgetsEditorModule.cpp
+++ b/Engine/Plugins/Runtime/AudioWidgets/Source/AudioWidgetsEditor/Private/AudioWidgetsEditorModule.cpp
@@ -2,14 +2,41 @@
 
 #include "AudioWidgetsEditorModule.h"
 #include "Modules/ModuleManager.h"
-
+#include "PropertyEditorModule.h"
+#include "SoundbaseDetails.h"
 
 void FAudioWidgetsEditorModule::StartupModule()
 {
+	RegisterCustomClassLayout("SoundBase", FOnGetDetailCustomizationInstance::CreateStatic(&FSoundBaseDetails::MakeInstance));
+
 }
 
 void FAudioWidgetsEditorModule::ShutdownModule()
 {
+	if (FModuleManager::Get().IsModuleLoaded("PropertyEditor"))
+	{
+		FPropertyEditorModule& PropertyModule = FModuleManager::GetModuleChecked<FPropertyEditorModule>("PropertyEditor");
+
+		// Unregister all classes customized by name
+		for (auto It = RegisteredClassNames.CreateConstIterator(); It; ++It)
+		{
+			if (It->IsValid())
+			{
+				PropertyModule.UnregisterCustomClassLayout(*It);
+			}
+		}
+	}
+}
+
+void FAudioWidgetsEditorModule::RegisterCustomClassLayout(FName ClassName, FOnGetDetailCustomizationInstance DetailLayoutDelegate )
+{
+	check(ClassName != NAME_None);
+
+	RegisteredClassNames.Add(ClassName);
+
+	const FName PropertyEditor("PropertyEditor");
+	FPropertyEditorModule& PropertyModule = FModuleManager::GetModuleChecked<FPropertyEditorModule>(PropertyEditor);
+	PropertyModule.RegisterCustomClassLayout( ClassName, DetailLayoutDelegate );
 }
 
 	
diff --git a/Engine/Plugins/Runtime/AudioWidgets/Source/AudioWidgetsEditor/Private/SoundBaseDetails.cpp b/Engine/Plugins/Runtime/AudioWidgets/Source/AudioWidgetsEditor/Private/SoundBaseDetails.cpp
new file mode 100644
--- /dev/null
+++ b/Engine/Plugins/Runtime/AudioWidgets/Source/AudioWidgetsEditor/Private/SoundBaseDetails.cpp
@@ -0,0 +1,24 @@
+// Copyright Epic Games, Inc. All Rights Reserved.
+
+#include "SoundBaseDetails.h"
+
+#include "DetailLayoutBuilder.h"
+#include "Modules/ModuleManager.h"
+#include "PropertyHandle.h"
+#include "Sound/AudioSettings.h"
+#include "Sound/SoundBase.h"
+#include "UObject/Class.h"
+#include "UObject/UObjectGlobals.h"
+
+TSharedRef<IDetailCustomization> FSoundBaseDetails::MakeInstance()
+{
+	return MakeShareable(new FSoundBaseDetails);
+}
+
+void FSoundBaseDetails::CustomizeDetails(IDetailLayoutBuilder& DetailBuilder)
+{
+	if (!FModuleManager::Get().IsModuleLoaded("AudioProperties"))
+	{
+		DetailBuilder.HideCategory("AudioProperties");
+	}
+}
diff --git a/Engine/Plugins/Runtime/AudioWidgets/Source/AudioWidgetsEditor/Private/SoundBaseDetails.h b/Engine/Plugins/Runtime/AudioWidgets/Source/AudioWidgetsEditor/Private/SoundBaseDetails.h
new file mode 100644
--- /dev/null
+++ b/Engine/Plugins/Runtime/AudioWidgets/Source/AudioWidgetsEditor/Private/SoundBaseDetails.h
@@ -0,0 +1,19 @@
+// Copyright Epic Games, Inc. All Rights Reserved.
+
+#pragma once
+
+#include "IDetailCustomization.h"
+#include "Templates/SharedPointer.h"
+
+class IDetailLayoutBuilder;
+
+class FSoundBaseDetails : public IDetailCustomization
+{
+public:
+	/** Makes a new instance of this detail layout class for a specific detail view requesting it */
+	static TSharedRef<IDetailCustomization> MakeInstance();
+
+private:
+	/** IDetailCustomization interface */
+	virtual void CustomizeDetails(IDetailLayoutBuilder& DetailBuilder) override;
+};
diff --git a/Engine/Plugins/Runtime/AudioWidgets/Source/AudioWidgetsEditor/Public/AudioWidgetsEditorModule.h b/Engine/Plugins/Runtime/AudioWidgets/Source/AudioWidgetsEditor/Public/AudioWidgetsEditorModule.h
--- a/Engine/Plugins/Runtime/AudioWidgets/Source/AudioWidgetsEditor/Public/AudioWidgetsEditorModule.h
+++ b/Engine/Plugins/Runtime/AudioWidgets/Source/AudioWidgetsEditor/Public/AudioWidgetsEditorModule.h
@@ -2,7 +2,10 @@
 
 #pragma once
 
+#include "Containers/Set.h"
 #include "Modules/ModuleInterface.h"
+#include "PropertyEditorDelegates.h"
+#include "UObject/NameTypes.h"
 
 class FAudioWidgetsEditorModule : public IModuleInterface
 {
@@ -11,4 +14,10 @@ public:
 	/** IModuleInterface implementation */
 	virtual void StartupModule() override;
 	virtual void ShutdownModule() override;
+	
+private:
+	void RegisterCustomClassLayout(FName ClassName, FOnGetDetailCustomizationInstance DetailLayoutDelegate);
+
+	/** List of registered class that we must unregister when the module shuts down */
+	TSet<FName> RegisteredClassNames;
 };
\ No newline at end of file
diff --git a/Engine/Source/Editor/DetailCustomizations/Private/DetailCustomizations.cpp b/Engine/Source/Editor/DetailCustomizations/Private/DetailCustomizations.cpp
--- a/Engine/Source/Editor/DetailCustomizations/Private/DetailCustomizations.cpp
+++ b/Engine/Source/Editor/DetailCustomizations/Private/DetailCustomizations.cpp
@@ -143,7 +143,6 @@
 #include "SlateSoundCustomization.h"
 #include "SoftClassPathCustomization.h"
 #include "SoftObjectPathCustomization.h"
-#include "SoundBaseDetails.h"
 #include "Sound/SoundNodeDistanceCrossFade.h"
 #include "SoundSourceBusDetails.h"
 #include "SoundWaveDetails.h"
@@ -437,7 +436,6 @@ void FDetailCustomizationsModule::RegisterObjectCustomizations()
 	RegisterCustomClassLayout("PoseAsset", FOnGetDetailCustomizationInstance::CreateStatic(&FPoseAssetDetails::MakeInstance));
 	RegisterCustomClassLayout("AnimationAsset", FOnGetDetailCustomizationInstance::CreateStatic(&FAnimationAssetDetails::MakeInstance));
 
-	RegisterCustomClassLayout("SoundBase", FOnGetDetailCustomizationInstance::CreateStatic(&FSoundBaseDetails::MakeInstance));
 	RegisterCustomClassLayout("SoundSourceBus", FOnGetDetailCustomizationInstance::CreateStatic(&FSoundSourceBusDetails::MakeInstance));
 	RegisterCustomClassLayout("DialogueWave", FOnGetDetailCustomizationInstance::CreateStatic(&FDialogueWaveDetails::MakeInstance));
 	RegisterCustomClassLayout("SoundWave", FOnGetDetailCustomizationInstance::CreateStatic(&FSoundWaveDetails::MakeInstance));
diff --git a/Engine/Source/Editor/DetailCustomizations/Private/SoundBaseDetails.cpp b/Engine/Source/Editor/DetailCustomizations/Private/SoundBaseDetails.cpp
deleted file mode 100644
diff --git a/Engine/Source/Editor/DetailCustomizations/Private/SoundBaseDetails.h b/Engine/Source/Editor/DetailCustomizations/Private/SoundBaseDetails.h
deleted file mode 100644
