[AudioWidgetsEditorModuke] correct include file name to fix for cis issue 591506
#rnx

diff --git a/Engine/Plugins/Runtime/AudioWidgets/Source/AudioWidgetsEditor/Private/AudioWidgetsEditorModule.cpp b/Engine/Plugins/Runtime/AudioWidgets/Source/AudioWidgetsEditor/Private/AudioWidgetsEditorModule.cpp
--- a/Engine/Plugins/Runtime/AudioWidgets/Source/AudioWidgetsEditor/Private/AudioWidgetsEditorModule.cpp
+++ b/Engine/Plugins/Runtime/AudioWidgets/Source/AudioWidgetsEditor/Private/AudioWidgetsEditorModule.cpp
@@ -3,7 +3,7 @@
 #include "AudioWidgetsEditorModule.h"
 #include "Modules/ModuleManager.h"
 #include "PropertyEditorModule.h"
-#include "SoundbaseDetails.h"
+#include "SoundBaseDetails.h"
 
 void FAudioWidgetsEditorModule::StartupModule()
 {
