[Metasounds] Create some example AudioMaterialSlate styles for users

#rnx
#virtualized

diff --git a/Engine/Plugins/Runtime/AudioWidgets/Source/AudioWidgets/Private/AudioMaterialSlate/AudioMaterialSlateTypes.cpp b/Engine/Plugins/Runtime/AudioWidgets/Source/AudioWidgets/Private/AudioMaterialSlate/AudioMaterialSlateTypes.cpp
--- a/Engine/Plugins/Runtime/AudioWidgets/Source/AudioWidgets/Private/AudioMaterialSlate/AudioMaterialSlateTypes.cpp
+++ b/Engine/Plugins/Runtime/AudioWidgets/Source/AudioWidgets/Private/AudioMaterialSlate/AudioMaterialSlateTypes.cpp
@@ -15,14 +15,56 @@ namespace AudioWidgets
 {
 	namespace SlateTypesPrivate
 	{
-		static const FLinearColor ButtonMainColor(0.098958f, 0.098958f, 0.098958f, 1.f);
-		static const FLinearColor ButtonAccentColor(0.341146f, 0.341146f, 0.341146f, 1.f);
-		static const FLinearColor ButtonPressedShadowColor(0.126558f, 0.138653f, 0.15f, 1.f);
-
-		static const FLinearColor BarMainColor(0.008f, 0.008f, 0.008f, 1.f);
-		static const FLinearColor BarAccentColor(0.005f, 0.005f, 0.005f, 1.f);
-		static const FLinearColor HandleMainColor(1.f, 1.f, 1.f, 1.f);
-		static const FLinearColor HandleOutlineColor(0.15f, 0.15f, 0.15f, 1.f);
+		// Button
+		static const FLinearColor ButtonMainColor(0.5f, 0.5f, 0.5f, 1.f);
+		static const FLinearColor ButtonAccentColor(0.96f, 0.96f, 0.96f, 1.f);
+		static const FLinearColor ButtonShadowColor(0.5f, 0.5f, 0.5f, 1.f);
+		static const FLinearColor ButtonUnpressedOutlineColor(0.0625f, 0.0625f, 0.0625f, 1.f);
+		static const FLinearColor ButtonPressedOutlineColor(0.96f, 0.96f, 0.96f, 1.f);
+		static const FLinearColor ButtonMainColorTint1(0.96f, 0.96f, 0.96f, 1.f);
+		static const FLinearColor ButtonMainColorTint2(0.06f, 0.06f, 0.06f, 1.f);
+		// End Button
+
+		// Slider
+		static const FLinearColor SliderBackgroundColor(0.008f, 0.008f, 0.008f, 1.f);
+		static const FLinearColor SliderBackgroundAccentColor(0.005f, 0.005f, 0.005f, 1.f);
+		static const FLinearColor SliderHandleMainColor(1.f, 1.f, 1.f, 1.f);
+		static const FLinearColor SliderHandleOutlineColor(0.15f, 0.15f, 0.15f, 1.f);
+		static const FLinearColor SliderValueMainColor(0.008f, 0.008f, 0.008f, 1.f);
+		// End Slider
+
+		// Knob
+		static const FLinearColor KnobMainColor(0.140625f, 0.140625f, 0.140625f, 1.f);
+		static const FLinearColor KnobAccentColor(0.06f, 0.06f, 0.06f, 1.f);
+		static const FLinearColor KnobShadowColor(0.06f, 0.06f, 0.06f, 1.f);
+		static const FLinearColor KnobSmoothBevelColor(0.041667f, 0.041667f, 0.041667f, 1.f);
+		static const FLinearColor KnobIndicatorDotColor(1.0f, 0.0f, 0.0f, 1.f);
+		static const FLinearColor KnobEdgeFillColor(0.015625f, 0.015625f, 0.015625f, 1.f);
+		static const FLinearColor KnobBarColor(0.067f, 0.067f, 0.067f, 1.f);
+		static const FLinearColor KnobBarShadowColor(0.067f, 0.067f, 0.067f, 1.f);
+		static const FLinearColor KnobBarFillMinColor(0.96f, 0.96f, 0.96f, 1.f);
+		static const FLinearColor KnobBarFillMidColor(0.96f, 0.96f, 0.96f, 1.f);
+		static const FLinearColor KnobBarFillMaxColor(0.96f, 0.96f, 0.96f, 1.f);
+		static const FLinearColor KnobBarFillTintColor(0.96f, 0.96f, 0.96f, 1.f);
+		// End Knob
+
+		// Meter
+		static const FLinearColor MeterFillMinColor(0.96f, 0.96f, 0.96f, 1.f);
+		static const FLinearColor MeterFillMidColor(0.96f, 0.96f, 0.96f, 1.f);
+		static const FLinearColor MeterFillMaxColor(0.96f, 0.96f, 0.96f, 1.f);
+		static const FLinearColor MeterFillBackgroundColor(0.06f, 0.06f, 0.06f, 1.f);
+		static const FVector2D MeterPadding(FVector2D(10.0f, 5.0f));
+		static const FVector2D MeterValueRangeDb(FVector2D(-60, 10));
+		static const bool bShowMeterScale(true);
+		static const bool bScaleMeterSide(true);
+		static const float MeterScaleHashOffset(5.0f);
+		static const float MeterScaleHashWidth(10.0f);
+		static const float MeterScaleHashHeight(1.0f);
+		static const int32 MeterDecibelsPerHash(5);
+		// End Meter
+
+		static const FSlateFontInfo Font(FStyleDefaults::GetFontInfo(5));
+		static const FAudioTextBoxStyle TextBoxStyle(FAudioTextBoxStyle::GetDefault());
 	}
 }
 
@@ -33,26 +75,21 @@ FAudioMaterialWidgetStyle::FAudioMaterialWidgetStyle()
 {
 }
 
-TObjectPtr<UMaterialInstanceDynamic> FAudioMaterialWidgetStyle::GetDynamicMaterial() const
-{
-	return DynamicMaterial;
-}
-
 UMaterialInstanceDynamic* FAudioMaterialWidgetStyle::CreateDynamicMaterial(UObject* InOuter) const
 {
-	DynamicMaterial = UMaterialInstanceDynamic::Create(Material, InOuter);
-	return DynamicMaterial;
+	return UMaterialInstanceDynamic::Create(Material, InOuter);
 }
 
 FAudioMaterialButtonStyle::FAudioMaterialButtonStyle()
 	: ButtonMainColor(SlateTypesPrivate::ButtonMainColor)
-	, ButtonShadowColor(FLinearColor::Black)
+	, ButtonMainColorTint_1(SlateTypesPrivate::ButtonMainColorTint1)
+	, ButtonMainColorTint_2(SlateTypesPrivate::ButtonMainColorTint2)
 	, ButtonAccentColor(SlateTypesPrivate::ButtonAccentColor)
-	, ButtonPressedMainColor(FLinearColor::White)
-	, ButtonPressedShadowColor(SlateTypesPrivate::ButtonPressedShadowColor)
-	, ButtonPressedOutlineColor(FLinearColor::Blue)
+	, ButtonShadowColor(SlateTypesPrivate::ButtonShadowColor)
+	, ButtonUnpressedOutlineColor(SlateTypesPrivate::ButtonUnpressedOutlineColor)
+	, ButtonPressedOutlineColor(SlateTypesPrivate::ButtonPressedOutlineColor)
 {
-	FString Path = PLUGIN_BASE_DIR + "MI_AudioMaterialButton.MI_AudioMaterialButton";
+	FString Path = PLUGIN_BASE_DIR + "MI_AudioMaterialToggleButton.MI_AudioMaterialToggleButton";
 	Material = LoadObject<UMaterialInterface>(nullptr, *Path);
 
 	DesiredSize = FVector2f(128.f, 128.f);	
@@ -71,12 +108,12 @@ const FAudioMaterialButtonStyle& FAudioMaterialButtonStyle::GetDefault()
 }
 
 FAudioMaterialSliderStyle::FAudioMaterialSliderStyle()
-	: BarMainColor(SlateTypesPrivate::BarMainColor)
-	, BarShadowColor(FLinearColor::Black)
-	, BarAccentColor(SlateTypesPrivate::BarAccentColor)
-	, HandleMainColor(SlateTypesPrivate::HandleMainColor)
-	, HandleOutlineColor(SlateTypesPrivate::HandleOutlineColor)
-	, TextBoxStyle(FAudioTextBoxStyle::GetDefault())
+	: SliderBackgroundColor(SlateTypesPrivate::SliderBackgroundColor)
+	, SliderBackgroundAccentColor(SlateTypesPrivate::SliderBackgroundAccentColor)
+	, SliderValueMainColor(SlateTypesPrivate::SliderValueMainColor)
+	, SliderHandleMainColor(SlateTypesPrivate::SliderHandleMainColor)
+	, SliderHandleOutlineColor(SlateTypesPrivate::SliderHandleOutlineColor)
+	, TextBoxStyle(SlateTypesPrivate::TextBoxStyle)
 {
 	FString Path = PLUGIN_BASE_DIR + "MI_AudioMaterialRoundedSlider.MI_AudioMaterialRoundedSlider";
 	Material = LoadObject<UMaterialInterface>(nullptr, *Path);
@@ -98,15 +135,18 @@ const FAudioMaterialSliderStyle& FAudioMaterialSliderStyle::GetDefault()
 }
 
 FAudioMaterialKnobStyle::FAudioMaterialKnobStyle()
-	: KnobMainColor(FLinearColor::Black)
-	, KnobAccentColor(FLinearColor::Gray)
-	, KnobIndicatorColor(FLinearColor::Red)
-	, KnobBarColor(FLinearColor::Gray)
-	, KnobBarShadowColor(FLinearColor::Gray)
-	, KnobBarFillMinColor(FLinearColor::White)
-	, KnobBarFillMidColor(FLinearColor::White)
-	, KnobBarFillMaxColor(FLinearColor::White)
-	, KnobBarFillTintColor(FLinearColor::White)
+	: KnobMainColor(SlateTypesPrivate::KnobMainColor)
+	, KnobAccentColor(SlateTypesPrivate::KnobAccentColor)
+	, KnobShadowColor(SlateTypesPrivate::KnobShadowColor)
+	, KnobSmoothBevelColor(SlateTypesPrivate::KnobSmoothBevelColor)
+	, KnobIndicatorDotColor(SlateTypesPrivate::KnobIndicatorDotColor)
+	, KnobEdgeFillColor(SlateTypesPrivate::KnobEdgeFillColor)
+	, KnobBarColor(SlateTypesPrivate::KnobBarColor)
+	, KnobBarShadowColor(SlateTypesPrivate::KnobBarShadowColor)
+	, KnobBarFillMinColor(SlateTypesPrivate::KnobBarFillMinColor)
+	, KnobBarFillMidColor(SlateTypesPrivate::KnobBarFillMidColor)
+	, KnobBarFillMaxColor(SlateTypesPrivate::KnobBarFillMaxColor)
+	, KnobBarFillTintColor(SlateTypesPrivate::KnobBarFillTintColor)
 {
 	FString Path = PLUGIN_BASE_DIR + "MI_AudioMaterialKnob.MI_AudioMaterialKnob";
 	Material = LoadObject<UMaterialInterface>(nullptr, *Path);
@@ -128,19 +168,19 @@ void FAudioMaterialKnobStyle::GetResources(TArray<const FSlateBrush*>& OutBrushe
 }
 
 FAudioMaterialMeterStyle::FAudioMaterialMeterStyle()
-	: MeterFillMinColor(FLinearColor::White)
-	, MeterFillMidColor(FLinearColor::White)
-	, MeterFillMaxColor(FLinearColor::White)
-	, MeterOffFillColor(FLinearColor::Black)
-	, MeterPadding(FVector2D(10.0f, 5.0f))
-	, ValueRangeDb(FVector2D(-60, 10))
-	, bShowScale(true)
-	, bScaleSide(true)
-	, ScaleHashOffset(5.0f)
-	, ScaleHashWidth(10.0f)
-	, ScaleHashHeight(1.0f)
-	, DecibelsPerHash(5)
-	, Font(FStyleDefaults::GetFontInfo(5))
+	: MeterFillMinColor(SlateTypesPrivate::MeterFillMinColor)
+	, MeterFillMidColor(SlateTypesPrivate::MeterFillMidColor)
+	, MeterFillMaxColor(SlateTypesPrivate::MeterFillMinColor)
+	, MeterFillBackgroundColor(SlateTypesPrivate::MeterFillBackgroundColor)
+	, MeterPadding(SlateTypesPrivate::MeterPadding)
+	, ValueRangeDb(SlateTypesPrivate::MeterValueRangeDb)
+	, bShowScale(SlateTypesPrivate::bShowMeterScale)
+	, bScaleSide(SlateTypesPrivate::bScaleMeterSide)
+	, ScaleHashOffset(SlateTypesPrivate::MeterScaleHashOffset)
+	, ScaleHashWidth(SlateTypesPrivate::MeterScaleHashWidth)
+	, ScaleHashHeight(SlateTypesPrivate::MeterScaleHashHeight)
+	, DecibelsPerHash(SlateTypesPrivate::MeterDecibelsPerHash)
+	, Font(SlateTypesPrivate::Font)
 {
 	FString Path = PLUGIN_BASE_DIR + "MI_AudioMaterialMeter.MI_AudioMaterialMeter";
 	Material = LoadObject<UMaterialInterface>(nullptr, *Path);
diff --git a/Engine/Plugins/Runtime/AudioWidgets/Source/AudioWidgets/Private/AudioMaterialSlate/SAudioMaterialButton.cpp b/Engine/Plugins/Runtime/AudioWidgets/Source/AudioWidgets/Private/AudioMaterialSlate/SAudioMaterialButton.cpp
--- a/Engine/Plugins/Runtime/AudioWidgets/Source/AudioWidgets/Private/AudioMaterialSlate/SAudioMaterialButton.cpp
+++ b/Engine/Plugins/Runtime/AudioWidgets/Source/AudioWidgets/Private/AudioMaterialSlate/SAudioMaterialButton.cpp
@@ -48,10 +48,11 @@ int32 SAudioMaterialButton::OnPaint(const FPaintArgs& Args, const FGeometry& All
 			DynamicMaterial.Get()->SetVectorParameterValue(FName("MainColor"), AudioMaterialButtonStyle->ButtonMainColor);
 			DynamicMaterial.Get()->SetVectorParameterValue(FName("ShadowColor"), AudioMaterialButtonStyle->ButtonShadowColor);
 			DynamicMaterial.Get()->SetVectorParameterValue(FName("SmoothBevelColor"), AudioMaterialButtonStyle->ButtonAccentColor);
-			DynamicMaterial.Get()->SetVectorParameterValue(FName("Color_1"), AudioMaterialButtonStyle->ButtonPressedMainColor);
-			DynamicMaterial.Get()->SetVectorParameterValue(FName("Color_2"), AudioMaterialButtonStyle->ButtonPressedShadowColor);
+			DynamicMaterial.Get()->SetVectorParameterValue(FName("Color_1"), AudioMaterialButtonStyle->ButtonMainColorTint_1);
+			DynamicMaterial.Get()->SetVectorParameterValue(FName("Color_2"), AudioMaterialButtonStyle->ButtonMainColorTint_2);
 			DynamicMaterial.Get()->SetVectorParameterValue(FName("LedColor"), AudioMaterialButtonStyle->ButtonPressedOutlineColor);
 			DynamicMaterial.Get()->SetScalarParameterValue(FName("Click"), bIsPressedAttribute.Get());
+			DynamicMaterial.Get()->SetVectorParameterValue(FName("BarColor"), AudioMaterialButtonStyle->ButtonUnpressedOutlineColor);
 
 			DynamicMaterial.Get()->SetScalarParameterValue(FName("LocalWidth"), AllottedGeometry.GetLocalSize().X);
 			DynamicMaterial.Get()->SetScalarParameterValue(FName("LocalHeigth"), AllottedGeometry.GetLocalSize().Y);		
diff --git a/Engine/Plugins/Runtime/AudioWidgets/Source/AudioWidgets/Private/AudioMaterialSlate/SAudioMaterialEnvelope.cpp b/Engine/Plugins/Runtime/AudioWidgets/Source/AudioWidgets/Private/AudioMaterialSlate/SAudioMaterialEnvelope.cpp
--- a/Engine/Plugins/Runtime/AudioWidgets/Source/AudioWidgets/Private/AudioMaterialSlate/SAudioMaterialEnvelope.cpp
+++ b/Engine/Plugins/Runtime/AudioWidgets/Source/AudioWidgets/Private/AudioMaterialSlate/SAudioMaterialEnvelope.cpp
@@ -25,10 +25,8 @@ int32 SAudioMaterialEnvelope::OnPaint(const FPaintArgs& Args, const FGeometry& A
 {
 	if (AudioMaterialEnvelopeStyle && EnvelopeSettings)
 	{
-		UMaterialInstanceDynamic* DynamicMaterial = AudioMaterialEnvelopeStyle->GetDynamicMaterial();
-		if (IsValid(DynamicMaterial))
+		if (DynamicMaterial.IsValid())
 		{
-
 			DynamicMaterial->SetScalarParameterValue(FName("A_Curve"), EnvelopeSettings->AttackCurve);
 			DynamicMaterial->SetScalarParameterValue(FName("A_Int"), EnvelopeSettings->AttackValue);
 			DynamicMaterial->SetScalarParameterValue(FName("A_Time"), EnvelopeSettings->AttackTime);
@@ -50,16 +48,23 @@ int32 SAudioMaterialEnvelope::OnPaint(const FPaintArgs& Args, const FGeometry& A
 
 			DynamicMaterial->SetScalarParameterValue(FName("LocalWidth"), AllottedGeometry.GetLocalSize().X);
 			DynamicMaterial->SetScalarParameterValue(FName("LocalHeigth"), AllottedGeometry.GetLocalSize().Y);			
-		}
-
-		const bool bEnabled = ShouldBeEnabled(bParentEnabled);
-		const ESlateDrawEffect DrawEffects = bEnabled ? ESlateDrawEffect::None : ESlateDrawEffect::DisabledEffect;
+			
+			const bool bEnabled = ShouldBeEnabled(bParentEnabled);
+			const ESlateDrawEffect DrawEffects = bEnabled ? ESlateDrawEffect::None : ESlateDrawEffect::DisabledEffect;
 
-		const FLinearColor FinalColorAndOpacity(InWidgetStyle.GetColorAndOpacityTint());
+			const FLinearColor FinalColorAndOpacity(InWidgetStyle.GetColorAndOpacityTint());
 
-		FSlateBrush Brush;
-		Brush.SetResourceObject(DynamicMaterial);
-		FSlateDrawElement::MakeBox(OutDrawElements, LayerId++, AllottedGeometry.ToPaintGeometry(), &Brush, DrawEffects, FinalColorAndOpacity);
+			FSlateBrush Brush;
+			Brush.SetResourceObject(DynamicMaterial.Get());
+			FSlateDrawElement::MakeBox(OutDrawElements, LayerId++, AllottedGeometry.ToPaintGeometry(), &Brush, DrawEffects, FinalColorAndOpacity);
+		}
+		else
+		{
+			if (AudioMaterialEnvelopeStyle)
+			{
+				DynamicMaterial = AudioMaterialEnvelopeStyle->CreateDynamicMaterial(Owner.Get());
+			}
+		}
 	}
 
 	return LayerId;
@@ -75,10 +80,12 @@ FVector2D SAudioMaterialEnvelope::ComputeDesiredSize(float) const
 	return FVector2D::ZeroVector;
 }
 
-void SAudioMaterialEnvelope::ApplyNewMaterial()
+UMaterialInstanceDynamic* SAudioMaterialEnvelope::ApplyNewMaterial()
 {
 	if (AudioMaterialEnvelopeStyle)
 	{
-		AudioMaterialEnvelopeStyle->CreateDynamicMaterial(Owner.Get());
+		DynamicMaterial = AudioMaterialEnvelopeStyle->CreateDynamicMaterial(Owner.Get());
 	}
+
+	return DynamicMaterial.Get();
 }
diff --git a/Engine/Plugins/Runtime/AudioWidgets/Source/AudioWidgets/Private/AudioMaterialSlate/SAudioMaterialKnob.cpp b/Engine/Plugins/Runtime/AudioWidgets/Source/AudioWidgets/Private/AudioMaterialSlate/SAudioMaterialKnob.cpp
--- a/Engine/Plugins/Runtime/AudioWidgets/Source/AudioWidgets/Private/AudioMaterialSlate/SAudioMaterialKnob.cpp
+++ b/Engine/Plugins/Runtime/AudioWidgets/Source/AudioWidgets/Private/AudioMaterialSlate/SAudioMaterialKnob.cpp
@@ -113,21 +113,19 @@ int32 SAudioMaterialKnob::OnPaint(const FPaintArgs& Args, const FGeometry& Allot
 		{
 			const float KnobPercent = ValueAttribute.Get();
 
-			//TODO unify material parameter names
-			DynamicMaterial.Get()->SetVectorParameterValue(FName("Color_2"), AudioMaterialKnobStyle->KnobMainColor);
-			DynamicMaterial.Get()->SetVectorParameterValue(FName("Color_1"), AudioMaterialKnobStyle->KnobAccentColor);
+			DynamicMaterial.Get()->SetVectorParameterValue(FName("Color_2"), AudioMaterialKnobStyle->KnobAccentColor);
+			DynamicMaterial.Get()->SetVectorParameterValue(FName("Color_1"), AudioMaterialKnobStyle->KnobMainColor);
 
 			DynamicMaterial.Get()->SetVectorParameterValue(FName("BarColor"), AudioMaterialKnobStyle->KnobBarColor);
 			DynamicMaterial.Get()->SetVectorParameterValue(FName("BarShadowColor"), AudioMaterialKnobStyle->KnobBarShadowColor);
-			DynamicMaterial.Get()->SetVectorParameterValue(FName("LEDglowMax"), AudioMaterialKnobStyle->KnobBarFillMaxColor);
 			DynamicMaterial.Get()->SetVectorParameterValue(FName("Led_Max"), AudioMaterialKnobStyle->KnobBarFillMaxColor);
-			DynamicMaterial.Get()->SetVectorParameterValue(FName("LEDglowMid"), AudioMaterialKnobStyle->KnobBarFillMidColor);
-			DynamicMaterial.Get()->SetVectorParameterValue(FName("Led_Mid"), AudioMaterialKnobStyle->KnobBarFillMidColor);
-			DynamicMaterial.Get()->SetVectorParameterValue(FName("LEDglowMin"), AudioMaterialKnobStyle->KnobBarFillMinColor);
+			DynamicMaterial.Get()->SetVectorParameterValue(FName("Led_Med"), AudioMaterialKnobStyle->KnobBarFillMidColor);
 			DynamicMaterial.Get()->SetVectorParameterValue(FName("LED_Min"), AudioMaterialKnobStyle->KnobBarFillMinColor);
-			DynamicMaterial.Get()->SetVectorParameterValue(FName("LineColor"), AudioMaterialKnobStyle->KnobIndicatorColor);
-			DynamicMaterial.Get()->SetVectorParameterValue(FName("DotColor"), AudioMaterialKnobStyle->KnobIndicatorColor);
+			DynamicMaterial.Get()->SetVectorParameterValue(FName("DotColor"), AudioMaterialKnobStyle->KnobIndicatorDotColor);
 			DynamicMaterial.Get()->SetVectorParameterValue(FName("LedTint"), AudioMaterialKnobStyle->KnobBarFillTintColor);
+			DynamicMaterial.Get()->SetVectorParameterValue(FName("EdgeFillColor"), AudioMaterialKnobStyle->KnobEdgeFillColor);
+			DynamicMaterial.Get()->SetVectorParameterValue(FName("ShadowColor"), AudioMaterialKnobStyle->KnobShadowColor);
+			DynamicMaterial.Get()->SetVectorParameterValue(FName("SmoothBevelColor"), AudioMaterialKnobStyle->KnobSmoothBevelColor);
 
 			DynamicMaterial.Get()->SetScalarParameterValue(FName("VALUE"), FMath::Clamp(KnobPercent, 0.f, 1.f));
 
@@ -278,4 +276,4 @@ void SAudioMaterialKnob::CommitValue(float NewValue)
 		Invalidate(EInvalidateWidgetReason::Paint);
 		OnValueChanged.ExecuteIfBound(Val);
 	}
-}
\ No newline at end of file
+}
diff --git a/Engine/Plugins/Runtime/AudioWidgets/Source/AudioWidgets/Private/AudioMaterialSlate/SAudioMaterialMeter.cpp b/Engine/Plugins/Runtime/AudioWidgets/Source/AudioWidgets/Private/AudioMaterialSlate/SAudioMaterialMeter.cpp
--- a/Engine/Plugins/Runtime/AudioWidgets/Source/AudioWidgets/Private/AudioMaterialSlate/SAudioMaterialMeter.cpp
+++ b/Engine/Plugins/Runtime/AudioWidgets/Source/AudioWidgets/Private/AudioMaterialSlate/SAudioMaterialMeter.cpp
@@ -57,8 +57,8 @@ int32 SAudioMaterialMeter::OnPaint(const FPaintArgs& Args, const FGeometry& Allo
 				DynamicMaterial->SetVectorParameterValue(FName("A (V3)"), Style->MeterFillMinColor);
 				DynamicMaterial->SetVectorParameterValue(FName("B (V3)"), Style->MeterFillMidColor);
 				DynamicMaterial->SetVectorParameterValue(FName("C (V3)"), Style->MeterFillMaxColor);
-				DynamicMaterial->SetVectorParameterValue(FName("OffColor"), Style->MeterOffFillColor);
-				DynamicMaterial->SetVectorParameterValue(FName("DotsOffColor"), Style->MeterOffFillColor);
+				DynamicMaterial->SetVectorParameterValue(FName("OffColor"), Style->MeterFillBackgroundColor);
+				DynamicMaterial->SetVectorParameterValue(FName("DotsOffColor"), Style->MeterFillBackgroundColor);
 
 				float ChannelMeterValueDb = FMath::GetMappedRangeValueClamped(Style->ValueRangeDb,FVector2D(0.f,1.f),  ChannelInfos[ChannelIndex].MeterValue);
 				DynamicMaterial->SetScalarParameterValue(FName("VALUE"), ChannelMeterValueDb);
diff --git a/Engine/Plugins/Runtime/AudioWidgets/Source/AudioWidgets/Private/AudioMaterialSlate/SAudioMaterialSlider.cpp b/Engine/Plugins/Runtime/AudioWidgets/Source/AudioWidgets/Private/AudioMaterialSlate/SAudioMaterialSlider.cpp
--- a/Engine/Plugins/Runtime/AudioWidgets/Source/AudioWidgets/Private/AudioMaterialSlate/SAudioMaterialSlider.cpp
+++ b/Engine/Plugins/Runtime/AudioWidgets/Source/AudioWidgets/Private/AudioMaterialSlate/SAudioMaterialSlider.cpp
@@ -44,18 +44,17 @@ int32 SAudioMaterialSlider::OnPaint(const FPaintArgs& Args, const FGeometry& All
 		if (DynamicMaterial.IsValid())
 		{
 			const float Value = ValueAttribute.Get();
-			DynamicMaterial.Get()->SetVectorParameterValue(FName("BarColor"), AudioMaterialSliderStyle->BarMainColor);
-			DynamicMaterial.Get()->SetVectorParameterValue(FName("LedColor"), AudioMaterialSliderStyle->BarShadowColor);
-			DynamicMaterial.Get()->SetVectorParameterValue(FName("BarInnerShadow"), AudioMaterialSliderStyle->BarAccentColor);
-			DynamicMaterial.Get()->SetVectorParameterValue(FName("ValueColor"), AudioMaterialSliderStyle->HandleMainColor);
-			DynamicMaterial.Get()->SetVectorParameterValue(FName("DotBevel1"), AudioMaterialSliderStyle->HandleOutlineColor);
+			DynamicMaterial.Get()->SetVectorParameterValue(FName("BarColor"), AudioMaterialSliderStyle->SliderBackgroundColor);
+			DynamicMaterial.Get()->SetVectorParameterValue(FName("LedColor"), AudioMaterialSliderStyle->SliderValueMainColor);
+			DynamicMaterial.Get()->SetVectorParameterValue(FName("BarInnerShadow"), AudioMaterialSliderStyle->SliderBackgroundAccentColor);
+			DynamicMaterial.Get()->SetVectorParameterValue(FName("ValueColor"), AudioMaterialSliderStyle->SliderHandleMainColor);
+			DynamicMaterial.Get()->SetVectorParameterValue(FName("ValueColorEdge"), AudioMaterialSliderStyle->SliderHandleOutlineColor);
 			DynamicMaterial.Get()->SetScalarParameterValue(FName("VALUE"), FMath::Clamp(Value, 0.f, 1.f));
 			DynamicMaterial.Get()->SetScalarParameterValue(FName("LedInt"), FMath::GetMappedRangeValueClamped(FVector2D(0.f,1.f),FVector2D(0.7f,4.f), Value));
 
 			DynamicMaterial.Get()->SetScalarParameterValue(FName("LocalWidth"), AllottedGeometry.GetLocalSize().X);
 			DynamicMaterial.Get()->SetScalarParameterValue(FName("LocalHeigth"), AllottedGeometry.GetLocalSize().Y);
 
-
 			// We draw the slider as a vertical slider regardless of the orientation and apply a render transform to make it display in the correct orientation.
 			FGeometry SliderGeometry = AllottedGeometry;
 			const float AllottedWidth = Orientation == Orient_Vertical ? AllottedGeometry.GetLocalSize().X : AllottedGeometry.GetLocalSize().Y;
diff --git a/Engine/Plugins/Runtime/AudioWidgets/Source/AudioWidgets/Public/AudioMaterialSlate/AudioMaterialSlateTypes.h b/Engine/Plugins/Runtime/AudioWidgets/Source/AudioWidgets/Public/AudioMaterialSlate/AudioMaterialSlateTypes.h
--- a/Engine/Plugins/Runtime/AudioWidgets/Source/AudioWidgets/Public/AudioMaterialSlate/AudioMaterialSlateTypes.h
+++ b/Engine/Plugins/Runtime/AudioWidgets/Source/AudioWidgets/Public/AudioMaterialSlate/AudioMaterialSlateTypes.h
@@ -16,7 +16,7 @@
 USTRUCT(BlueprintType)
 struct AUDIOWIDGETS_API FAudioMaterialWidgetStyle : public FSlateWidgetStyle
 {
-	GENERATED_USTRUCT_BODY()
+	GENERATED_BODY()
 
 public:
 
@@ -33,12 +33,7 @@ public:
 public:
 
 	UMaterialInstanceDynamic* CreateDynamicMaterial(UObject* InOuter) const;
-	TObjectPtr<UMaterialInstanceDynamic> GetDynamicMaterial() const;
 
-private:
-
-	UPROPERTY(meta = (DeprecatedProperty, DeprecationMessage = "keep the reference instead as a TWeakObjectPtr in the AudioMaterialSlates.  This will be removed when all AudioMaterialSlates are updated."))
-	mutable TObjectPtr<UMaterialInstanceDynamic> DynamicMaterial;
 };
 
 /**
@@ -47,7 +42,7 @@ private:
 USTRUCT(BlueprintType)
 struct AUDIOWIDGETS_API FAudioMaterialButtonStyle : public FAudioMaterialWidgetStyle
 {
-	GENERATED_USTRUCT_BODY()
+	GENERATED_BODY()
 
 	FAudioMaterialButtonStyle();
 
@@ -59,27 +54,38 @@ struct AUDIOWIDGETS_API FAudioMaterialButtonStyle : public FAudioMaterialWidgetS
 
 	FAudioMaterialButtonStyle& SetMaterial(UMaterialInterface* InMaterialInterface) { Material = InMaterialInterface; return *this; }
 
-	UPROPERTY(EditAnywhere, BlueprintReadWrite, Category = "Style")
+	/** The button's Main color value. */
+	UPROPERTY(EditAnywhere, BlueprintReadWrite, Category = "Style | Button")
 	FLinearColor ButtonMainColor;
 	FAudioMaterialButtonStyle& SetButtonMainColor(const FLinearColor& InColor) { ButtonMainColor = InColor; return *this; }
 
-	UPROPERTY(EditAnywhere, BlueprintReadWrite, Category = "Style")
-	FLinearColor ButtonShadowColor;
-	FAudioMaterialButtonStyle& SetButtonShadowColor(const FLinearColor& InColor) { ButtonShadowColor = InColor; return *this; }
+	/** The button color's Tint value covering one half of the gradient. */
+	UPROPERTY(EditAnywhere, BlueprintReadWrite, Category = "Style | Button")
+	FLinearColor ButtonMainColorTint_1;
+	FAudioMaterialButtonStyle& SetButtonMainColorTint_1(const FLinearColor& InColor) { ButtonMainColorTint_1 = InColor; return *this; }
 
-	UPROPERTY(EditAnywhere, BlueprintReadWrite, Category = "Style")
+	/** The button color's Tint value covering the other half of the gradient. */
+	UPROPERTY(EditAnywhere, BlueprintReadWrite, Category = "Style | Button")
+	FLinearColor ButtonMainColorTint_2;
+	FAudioMaterialButtonStyle& SetButtonMainColorTint_2(const FLinearColor& InColor) { ButtonMainColorTint_2 = InColor; return *this; }	
+
+	/** The button's Accent color value. */
+	UPROPERTY(EditAnywhere, BlueprintReadWrite, Category = "Style | Button")
 	FLinearColor ButtonAccentColor;	
-	FAudioMaterialButtonStyle& SetButtonAccentColor(const FLinearColor& InColor) { ButtonShadowColor = InColor; return *this; }
-	
-	UPROPERTY(EditAnywhere, BlueprintReadWrite, Category = "Style")
-	FLinearColor ButtonPressedMainColor;	
-	FAudioMaterialButtonStyle& SetButtonPressedMainColor(const FLinearColor& InColor) { ButtonPressedMainColor = InColor; return *this; }
+	FAudioMaterialButtonStyle& SetButtonAccentColor(const FLinearColor& InColor) { ButtonAccentColor = InColor; return *this; }
 
-	UPROPERTY(EditAnywhere, BlueprintReadWrite, Category = "Style")
-	FLinearColor ButtonPressedShadowColor;	
-	FAudioMaterialButtonStyle& SetButtonPressedShadowColor(const FLinearColor& InColor) { ButtonPressedShadowColor = InColor; return *this; }
+	/** The button's Shadow color value. */
+	UPROPERTY(EditAnywhere, BlueprintReadWrite, Category = "Style | Button")
+	FLinearColor ButtonShadowColor;	
+	FAudioMaterialButtonStyle& SetButtonShadowColor(const FLinearColor& InColor) { ButtonShadowColor = InColor; return *this; }
 
-	UPROPERTY(EditAnywhere, BlueprintReadWrite, Category = "Style")
+	/** The button's Outline color value when Unpressed. */
+	UPROPERTY(EditAnywhere, BlueprintReadWrite, Category = "Style | Button")
+	FLinearColor ButtonUnpressedOutlineColor;
+	FAudioMaterialButtonStyle& SetButtonUnpressedOutlineColor(const FLinearColor& InColor) { ButtonUnpressedOutlineColor = InColor; return *this; }
+
+	/** The button's Outline color value when Pressed. */
+	UPROPERTY(EditAnywhere, BlueprintReadWrite, Category = "Style | Button")
 	FLinearColor ButtonPressedOutlineColor;
 	FAudioMaterialButtonStyle& SetButtonPressedOutlineColor(const FLinearColor& InColor) { ButtonPressedOutlineColor = InColor; return *this; }
 
@@ -91,7 +97,7 @@ struct AUDIOWIDGETS_API FAudioMaterialButtonStyle : public FAudioMaterialWidgetS
 USTRUCT(BlueprintType)
 struct AUDIOWIDGETS_API FAudioMaterialSliderStyle : public FAudioMaterialWidgetStyle
 {
-	GENERATED_USTRUCT_BODY()
+	GENERATED_BODY()
 
 	FAudioMaterialSliderStyle();
 
@@ -103,28 +109,33 @@ struct AUDIOWIDGETS_API FAudioMaterialSliderStyle : public FAudioMaterialWidgetS
 
 	FAudioMaterialSliderStyle& SetMaterial(UMaterialInterface* InMaterialInterface) { Material = InMaterialInterface; return *this; }
 
-	UPROPERTY(EditAnywhere, BlueprintReadWrite, Category = "Style")
-	FLinearColor BarMainColor;
-	FAudioMaterialSliderStyle& SetSliderBarMainColor(const FLinearColor& InColor) { BarMainColor = InColor; return *this; }
-	
-	UPROPERTY(EditAnywhere, BlueprintReadWrite, Category = "Style")
-	FLinearColor BarShadowColor;	
-	FAudioMaterialSliderStyle& SetSliderBarShadowColor(const FLinearColor& InColor) { BarShadowColor = InColor; return *this; }
-	
-	UPROPERTY(EditAnywhere, BlueprintReadWrite, Category = "Style")
-	FLinearColor BarAccentColor;
-	FAudioMaterialSliderStyle& SetSliderBarAccentColor(const FLinearColor& InColor) { BarAccentColor = InColor; return *this; }
-
-	UPROPERTY(EditAnywhere, BlueprintReadWrite, Category = "Style")
-	FLinearColor HandleMainColor;	
-	FAudioMaterialSliderStyle& SetSliderHandleMainColor(const FLinearColor& InColor) { HandleMainColor = InColor; return *this; }
+	/** The slider Bar's Background color value. */
+	UPROPERTY(EditAnywhere, BlueprintReadWrite, Category = "Style | Slider")
+	FLinearColor SliderBackgroundColor;
+	FAudioMaterialSliderStyle& SetSliderBarBackgroundColor(const FLinearColor& InColor) { SliderBackgroundColor = InColor; return *this; }
+
+	/** The slider Bar's Background Accent color value. Can be thought as the slider's Inner Shadow color value. */
+	UPROPERTY(EditAnywhere, BlueprintReadWrite, Category = "Style | Slider")
+	FLinearColor SliderBackgroundAccentColor;
+	FAudioMaterialSliderStyle& SetSliderBarBackgroundAccentColor(const FLinearColor& InColor) { SliderBackgroundAccentColor = InColor; return *this; }
+
+	/** The slider's Color value representing the slider's Output Value amount. */
+	UPROPERTY(EditAnywhere, BlueprintReadWrite, Category = "Style | Slider")
+	FLinearColor SliderValueMainColor;
+	FAudioMaterialSliderStyle& SetSliderBarValueMainColor(const FLinearColor& InColor) { SliderValueMainColor = InColor; return *this; }
+
+	/** The slider Handle's Main color value. */
+	UPROPERTY(EditAnywhere, BlueprintReadWrite, Category = "Style | Slider")
+	FLinearColor SliderHandleMainColor;
+	FAudioMaterialSliderStyle& SetSliderHandleMainColor(const FLinearColor& InColor) { SliderHandleMainColor = InColor; return *this; }
 	
-	UPROPERTY(EditAnywhere, BlueprintReadWrite, Category = "Style")
-	FLinearColor HandleOutlineColor;
-	FAudioMaterialSliderStyle& SetSliderHandleOutlineColor(const FLinearColor& InColor) { HandleOutlineColor = InColor; return *this; }
+	/** The slider Handle's Outline color value. */
+	UPROPERTY(EditAnywhere, BlueprintReadWrite, Category = "Style | Slider")
+	FLinearColor SliderHandleOutlineColor;
+	FAudioMaterialSliderStyle& SetSliderHandleOutlineColor(const FLinearColor& InColor) { SliderHandleOutlineColor = InColor; return *this; }
 
 	/** The style to use for the audio text box widget. */
-	UPROPERTY(EditAnywhere, BlueprintReadWrite, Category = "Style")
+	UPROPERTY(EditAnywhere, BlueprintReadWrite, Category = "Style | Label")
 	FAudioTextBoxStyle TextBoxStyle;
 	FAudioMaterialSliderStyle& SetTextBoxStyle(const FAudioTextBoxStyle& InTextBoxStyle) { TextBoxStyle = InTextBoxStyle; return *this; }
 
@@ -136,7 +147,7 @@ struct AUDIOWIDGETS_API FAudioMaterialSliderStyle : public FAudioMaterialWidgetS
 USTRUCT(BlueprintType)
 struct AUDIOWIDGETS_API FAudioMaterialKnobStyle : public FAudioMaterialWidgetStyle
 {
-	GENERATED_USTRUCT_BODY()
+	GENERATED_BODY()
 
 	FAudioMaterialKnobStyle();
 
@@ -148,46 +159,71 @@ struct AUDIOWIDGETS_API FAudioMaterialKnobStyle : public FAudioMaterialWidgetSty
 
 	FAudioMaterialKnobStyle& SetMaterial(UMaterialInterface* InMaterialInterface) { Material = InMaterialInterface; return *this; }
 
-	UPROPERTY(EditAnywhere, BlueprintReadWrite, Category = "Style")
+	/** The knob's Main color value. */
+	UPROPERTY(EditAnywhere, BlueprintReadWrite, Category = "Style | Knob")
 	FLinearColor KnobMainColor;
 	FAudioMaterialKnobStyle& SetKnobMainColor(const FLinearColor& InColor) { KnobMainColor = InColor; return *this; }
 
-	UPROPERTY(EditAnywhere, BlueprintReadWrite, Category = "Style")
+	/** The knob's Accent color value. */
+	UPROPERTY(EditAnywhere, BlueprintReadWrite, Category = "Style | Knob")
 	FLinearColor KnobAccentColor;
 	FAudioMaterialKnobStyle& SetKnobAccentColor(const FLinearColor& InColor) { KnobAccentColor = InColor; return *this; }
 	
-	UPROPERTY(EditAnywhere, BlueprintReadWrite, Category = "Style")
-	FLinearColor KnobIndicatorColor;
-	FAudioMaterialKnobStyle& SetKnobIndicatorColor(const FLinearColor& InColor) {KnobIndicatorColor = InColor; return *this;}
+	/** The knob's Shadow color value. */
+	UPROPERTY(EditAnywhere, BlueprintReadWrite, Category = "Style | Knob")
+	FLinearColor KnobShadowColor;
+	FAudioMaterialKnobStyle& SetKnobShadowColor(const FLinearColor& InColor) { KnobShadowColor = InColor; return *this; }
 	
-	UPROPERTY(EditAnywhere, BlueprintReadWrite, Category = "Style")
+	/** The knob's Smooth Bevel color value. */
+	UPROPERTY(EditAnywhere, BlueprintReadWrite, Category = "Style | Knob")
+	FLinearColor KnobSmoothBevelColor;
+	FAudioMaterialKnobStyle& SetSmoothBevelColor(const FLinearColor& InColor) { KnobSmoothBevelColor = InColor; return *this; }
+	
+	/** The knob's Indicator Dot color value. */
+	UPROPERTY(EditAnywhere, BlueprintReadWrite, Category = "Style | Knob")
+	FLinearColor KnobIndicatorDotColor;
+	FAudioMaterialKnobStyle& SetKnobIndicatorColor(const FLinearColor& InColor) { KnobIndicatorDotColor = InColor; return *this;}
+	
+	/* The knob's Edge Fill color value. */
+	UPROPERTY(EditAnywhere, BlueprintReadWrite, Category = "Style | Knob")
+	FLinearColor KnobEdgeFillColor;
+	FAudioMaterialKnobStyle& SetKnobEdgeFillColor(const FLinearColor& InColor) { KnobEdgeFillColor = InColor; return *this;}
+	
+	/** The knob Bar's Color value. */
+	UPROPERTY(EditAnywhere, BlueprintReadWrite, Category = "Style | Knob")
 	FLinearColor KnobBarColor;
 	FAudioMaterialKnobStyle& SetKnobBarColor(const FLinearColor& InColor) { KnobBarColor = InColor; return *this; }
 	
-	UPROPERTY(EditAnywhere, BlueprintReadWrite, Category = "Style")
+	/** The knob Bar's Shadow color value. */
+	UPROPERTY(EditAnywhere, BlueprintReadWrite, Category = "Style | Knob")
 	FLinearColor KnobBarShadowColor;
 	FAudioMaterialKnobStyle& SetKnobBarShadowColor(const FLinearColor& InColor) { KnobBarShadowColor = InColor; return *this; }
 
-	UPROPERTY(EditAnywhere, BlueprintReadWrite, Category = "Style")
+	/** The knob Bar's Fill color value representing the Starting section of the fill. */
+	UPROPERTY(EditAnywhere, BlueprintReadWrite, Category = "Style | Knob")
 	FLinearColor KnobBarFillMinColor;
 	FAudioMaterialKnobStyle& SetKnobBarFillMinColor(const FLinearColor& InColor) { KnobBarFillMinColor = InColor; return *this; }
 	
-	UPROPERTY(EditAnywhere, BlueprintReadWrite, Category = "Style")
+	/** The knob Bar's Fill color value representing the Middle section of the fill. */
+	UPROPERTY(EditAnywhere, BlueprintReadWrite, Category = "Style | Knob")
 	FLinearColor KnobBarFillMidColor;
 	FAudioMaterialKnobStyle& SetKnobFillMidColor(const FLinearColor& InColor) { KnobBarFillMidColor = InColor; return *this; }
 	
-	UPROPERTY(EditAnywhere, BlueprintReadWrite, Category = "Style")
+	/** The knob Bar's Fill color value representing the Ending section of the fill. */
+	UPROPERTY(EditAnywhere, BlueprintReadWrite, Category = "Style | Knob")
 	FLinearColor KnobBarFillMaxColor;
 	FAudioMaterialKnobStyle& SetKnobBarFillMaxColor(const FLinearColor& InColor) { KnobBarFillMaxColor = InColor; return *this; }
 	
-	UPROPERTY(EditAnywhere, BlueprintReadWrite, Category = "Style")
+	/** The knob Bar Fill color's Tint value. */
+	UPROPERTY(EditAnywhere, BlueprintReadWrite, Category = "Style | Knob")
 	FLinearColor KnobBarFillTintColor;
 	FAudioMaterialKnobStyle& SetKnobBarFillTintColor(const FLinearColor& InColor) { KnobBarFillTintColor = InColor; return *this; }
 
 	/** The style to use for the audio text box widget. */
-	UPROPERTY(EditAnywhere, BlueprintReadWrite, Category = "Style")
+	UPROPERTY(EditAnywhere, BlueprintReadWrite, Category = "Style | Label")
 	FAudioTextBoxStyle TextBoxStyle;
 	void SetTextBoxStyle(const FAudioTextBoxStyle& InTextBoxStyle) { TextBoxStyle = InTextBoxStyle; }
+
 };
 
 /**
@@ -196,7 +232,7 @@ struct AUDIOWIDGETS_API FAudioMaterialKnobStyle : public FAudioMaterialWidgetSty
 USTRUCT(BlueprintType)
 struct AUDIOWIDGETS_API FAudioMaterialMeterStyle : public FAudioMaterialWidgetStyle
 {
-	GENERATED_USTRUCT_BODY()
+	GENERATED_BODY()
 
 	FAudioMaterialMeterStyle();
 
@@ -208,21 +244,25 @@ struct AUDIOWIDGETS_API FAudioMaterialMeterStyle : public FAudioMaterialWidgetSt
 
 	FAudioMaterialMeterStyle& SetMaterial(UMaterialInterface* InMaterialInterface) { Material = InMaterialInterface; return *this; }
 
+	/** The meter's Fill color value representing the Starting section of the fill. */
 	UPROPERTY(EditAnywhere, BlueprintReadWrite, Category = "Style")
 	FLinearColor MeterFillMinColor;
 	FAudioMaterialMeterStyle& SetMeterFillMinColor(const FLinearColor& InColor) { MeterFillMinColor = InColor; return *this; }
 
+	/** The meter's Fill color value representing the Middle section of the fill. */
 	UPROPERTY(EditAnywhere, BlueprintReadWrite, Category = "Style")
 	FLinearColor MeterFillMidColor;
 	FAudioMaterialMeterStyle& SetMeterFillMidColor(const FLinearColor& InColor) { MeterFillMidColor = InColor; return *this; }
 
+	/** The meter's Fill color value representing the Ending section of the fill. */
 	UPROPERTY(EditAnywhere, BlueprintReadWrite, Category = "Style")
 	FLinearColor MeterFillMaxColor;
 	FAudioMaterialMeterStyle& SetMeterFillMaxColor(const FLinearColor& InColor) { MeterFillMaxColor = InColor; return *this; }
 
+	/** The meter's Background Fill color value. */
 	UPROPERTY(EditAnywhere, BlueprintReadWrite, Category = "Style")
-	FLinearColor MeterOffFillColor;
-	FAudioMaterialMeterStyle& SetMeterOffFillColor(const FLinearColor& InColor) { MeterOffFillColor = InColor; return *this; }
+	FLinearColor MeterFillBackgroundColor;
+	FAudioMaterialMeterStyle& SetMeterOffFillColor(const FLinearColor& InColor) { MeterFillBackgroundColor = InColor; return *this; }
 
 	// How much padding to add around the meter
 	UPROPERTY(EditAnywhere, BlueprintReadWrite, Category = Appearance)
@@ -289,7 +329,7 @@ struct AUDIOWIDGETS_API FAudioMaterialMeterStyle : public FAudioMaterialWidgetSt
 USTRUCT(BlueprintType)
 struct AUDIOWIDGETS_API FAudioMaterialEnvelopeStyle : public FAudioMaterialWidgetStyle
 {
-	GENERATED_USTRUCT_BODY()
+	GENERATED_BODY()
 
 	FAudioMaterialEnvelopeStyle();
 
@@ -313,4 +353,4 @@ struct AUDIOWIDGETS_API FAudioMaterialEnvelopeStyle : public FAudioMaterialWidge
 	FLinearColor OutlineColor;
 	FAudioMaterialEnvelopeStyle& SetEnvelopeOutlineColor(const FLinearColor& InColor) { OutlineColor = InColor; return *this; }
 	
-};
\ No newline at end of file
+};
diff --git a/Engine/Plugins/Runtime/AudioWidgets/Source/AudioWidgets/Public/AudioMaterialSlate/SAudioMaterialEnvelope.h b/Engine/Plugins/Runtime/AudioWidgets/Source/AudioWidgets/Public/AudioMaterialSlate/SAudioMaterialEnvelope.h
--- a/Engine/Plugins/Runtime/AudioWidgets/Source/AudioWidgets/Public/AudioMaterialSlate/SAudioMaterialEnvelope.h
+++ b/Engine/Plugins/Runtime/AudioWidgets/Source/AudioWidgets/Public/AudioMaterialSlate/SAudioMaterialEnvelope.h
@@ -40,7 +40,7 @@ public:
 	virtual FVector2D ComputeDesiredSize(float) const override;
 
 	/** Apply new material to be used to render the Slate.*/
-	void ApplyNewMaterial();
+	UMaterialInstanceDynamic* ApplyNewMaterial();
 
 private:
 
@@ -50,6 +50,9 @@ private:
 	// Holds the style for the Slate
 	const FAudioMaterialEnvelopeStyle* AudioMaterialEnvelopeStyle = nullptr;
 
+	// Holds the Modifiable Material that represent the Envelope
+	mutable TWeakObjectPtr<UMaterialInstanceDynamic> DynamicMaterial;
+
 	//Holds the current Envelope settings
 	const FAudioMaterialEnvelopeSettings* EnvelopeSettings = nullptr;
 
