[AudioWidgets] Added analyzer widget styles to FAudioWidgetsStyle for more consistent use of styles across all of the widgets.

diff --git a/Engine/Plugins/Runtime/AudioWidgets/Source/AudioWidgets/Private/AudioMeter.cpp b/Engine/Plugins/Runtime/AudioWidgets/Source/AudioWidgets/Private/AudioMeter.cpp
--- a/Engine/Plugins/Runtime/AudioWidgets/Source/AudioWidgets/Private/AudioMeter.cpp
+++ b/Engine/Plugins/Runtime/AudioWidgets/Source/AudioWidgets/Private/AudioMeter.cpp
@@ -32,7 +32,7 @@ UAudioMeter::UAudioMeter(const FObjectInitializer& ObjectInitializer)
 	DefaultInfo.PeakValue = -3.0f;
 	MeterChannelInfo.Add(DefaultInfo);
 
-	WidgetStyle = FAudioMeterStyle::GetDefault();
+	WidgetStyle = FAudioWidgetsStyle::Get().GetWidgetStyle<FAudioMeterStyle>("AudioMeter.Style");
 
 #if WITH_EDITORONLY_DATA
 	AccessibleBehavior = ESlateAccessibleBehavior::NotAccessible;
@@ -160,6 +160,8 @@ const FText UAudioMeter::GetPaletteCategory()
 }
 #endif
 
+const FName FAudioMeterDefaultColorStyle::TypeName(TEXT("FAudioMeterDefaultColorStyle"));
+
 const FAudioMeterDefaultColorStyle& FAudioMeterDefaultColorStyle::GetDefault()
 {
 	static FAudioMeterDefaultColorStyle Style;
@@ -175,7 +177,7 @@ namespace AudioWidgets
 
 	FAudioMeter::FAudioMeter(const int32 InNumChannels, const Audio::FDeviceId InAudioDeviceId, const TObjectPtr<UAudioBus> InExternalAudioBus, const FAudioMeterDefaultColorStyle* AudioMeterWidgetStyle)
 	{
-		const FAudioMeterDefaultColorStyle* WidgetStyle = AudioMeterWidgetStyle ? AudioMeterWidgetStyle : &FAudioMeterDefaultColorStyle::GetDefault();
+		const FAudioMeterDefaultColorStyle* WidgetStyle = AudioMeterWidgetStyle ? AudioMeterWidgetStyle : &FAudioWidgetsStyle::Get().GetWidgetStyle<FAudioMeterDefaultColorStyle>("AudioMeter.DefaultColorStyle");
 		Widget = SNew(SAudioMeter)
 			.Orientation(EOrientation::Orient_Vertical)
 			.BackgroundColor(FLinearColor::Transparent)
diff --git a/Engine/Plugins/Runtime/AudioWidgets/Source/AudioWidgets/Private/AudioOscilloscope.cpp b/Engine/Plugins/Runtime/AudioWidgets/Source/AudioWidgets/Private/AudioOscilloscope.cpp
--- a/Engine/Plugins/Runtime/AudioWidgets/Source/AudioWidgets/Private/AudioOscilloscope.cpp
+++ b/Engine/Plugins/Runtime/AudioWidgets/Source/AudioWidgets/Private/AudioOscilloscope.cpp
@@ -12,7 +12,7 @@ namespace AudioWidgets
 		const EAudioPanelLayoutType InPanelLayoutType,
 		const FAudioOscilloscopePanelStyle* InOscilloscopePanelStyle)
 	{
-		OscilloscopePanelStyle = InOscilloscopePanelStyle ? *InOscilloscopePanelStyle : FAudioOscilloscopePanelStyle::GetDefault();
+		OscilloscopePanelStyle = InOscilloscopePanelStyle ? *InOscilloscopePanelStyle : FAudioWidgetsStyle::Get().GetWidgetStyle<FAudioOscilloscopePanelStyle>("AudioOscilloscope.PanelStyle");
 		CreateAudioBus(InNumChannels);
 		CreateDataProvider(InAudioDeviceId, InTimeWindowMs, InMaxTimeWindowMs, InAnalysisPeriodMs, InPanelLayoutType);
 		CreateOscilloscopeWidget(InNumChannels, InPanelLayoutType);
diff --git a/Engine/Plugins/Runtime/AudioWidgets/Source/AudioWidgets/Private/AudioOscilloscopeUMG.cpp b/Engine/Plugins/Runtime/AudioWidgets/Source/AudioWidgets/Private/AudioOscilloscopeUMG.cpp
--- a/Engine/Plugins/Runtime/AudioWidgets/Source/AudioWidgets/Private/AudioOscilloscopeUMG.cpp
+++ b/Engine/Plugins/Runtime/AudioWidgets/Source/AudioWidgets/Private/AudioOscilloscopeUMG.cpp
@@ -13,7 +13,7 @@
 
 UAudioOscilloscope::UAudioOscilloscope(const FObjectInitializer& ObjectInitializer)
 	: Super(ObjectInitializer)
-	, OscilloscopeStyle(FAudioOscilloscopePanelStyle::GetDefault())
+	, OscilloscopeStyle(FAudioWidgetsStyle::Get().GetWidgetStyle<FAudioOscilloscopePanelStyle>("AudioOscilloscope.PanelStyle"))
 {
 #if WITH_EDITORONLY_DATA
 	AccessibleBehavior       = ESlateAccessibleBehavior::NotAccessible;
diff --git a/Engine/Plugins/Runtime/AudioWidgets/Source/AudioWidgets/Private/AudioSpectrumAnalyzer.cpp b/Engine/Plugins/Runtime/AudioWidgets/Source/AudioWidgets/Private/AudioSpectrumAnalyzer.cpp
--- a/Engine/Plugins/Runtime/AudioWidgets/Source/AudioWidgets/Private/AudioSpectrumAnalyzer.cpp
+++ b/Engine/Plugins/Runtime/AudioWidgets/Source/AudioWidgets/Private/AudioSpectrumAnalyzer.cpp
@@ -47,7 +47,7 @@ namespace AudioWidgets
 		: SpectrumAnalysisSettings(NewObject<USynesthesiaSpectrumAnalysisSettings>())
 		, ConstantQSettings(NewObject<UConstantQSettings>())
 		, Widget(SNew(SAudioSpectrumPlot)
-			.Style(Params.PlotStyle ? Params.PlotStyle : &FAudioSpectrumPlotStyle::GetDefault())
+			.Style(Params.PlotStyle ? Params.PlotStyle : &FAudioWidgetsStyle::Get().GetWidgetStyle<FAudioSpectrumPlotStyle>("AudioSpectrumPlot.Style"))
 			.Clipping(EWidgetClipping::ClipToBounds)
 			.TiltExponent(Params.TiltExponent)
 			.DisplayCrosshair(true)
diff --git a/Engine/Plugins/Runtime/AudioWidgets/Source/AudioWidgets/Private/AudioVectorscope.cpp b/Engine/Plugins/Runtime/AudioWidgets/Source/AudioWidgets/Private/AudioVectorscope.cpp
--- a/Engine/Plugins/Runtime/AudioWidgets/Source/AudioWidgets/Private/AudioVectorscope.cpp
+++ b/Engine/Plugins/Runtime/AudioWidgets/Source/AudioWidgets/Private/AudioVectorscope.cpp
@@ -14,7 +14,7 @@ namespace AudioWidgets
 		const EAudioPanelLayoutType InPanelLayoutType,
 		const FAudioVectorscopePanelStyle* InPanelStyle)
 	{
-		VectorscopePanelStyle = InPanelStyle ? *InPanelStyle : FAudioVectorscopePanelStyle::GetDefault();
+		VectorscopePanelStyle = InPanelStyle ? *InPanelStyle : FAudioWidgetsStyle::Get().GetWidgetStyle<FAudioVectorscopePanelStyle>("AudioVectorscope.PanelStyle");
 		CreateAudioBus(InNumChannels);
 		CreateDataProvider(InAudioDeviceId, InTimeWindowMs, InMaxTimeWindowMs, InAnalysisPeriodMs);
 		CreateVectorscopeWidget(InPanelLayoutType);
@@ -47,7 +47,6 @@ namespace AudioWidgets
 		if (!VectorscopePanelWidget.IsValid())
 		{
 			VectorscopePanelWidget = SNew(SAudioVectorscopePanelWidget, SequenceView)
-				.PanelStyle(PanelStyle ? PanelStyle : &FAudioVectorscopePanelStyle::GetDefault())
 				.PanelLayoutType(InPanelLayoutType)
 				.PanelStyle(&VectorscopePanelStyle);
 		}
diff --git a/Engine/Plugins/Runtime/AudioWidgets/Source/AudioWidgets/Private/AudioVectorscopeUMG.cpp b/Engine/Plugins/Runtime/AudioWidgets/Source/AudioWidgets/Private/AudioVectorscopeUMG.cpp
--- a/Engine/Plugins/Runtime/AudioWidgets/Source/AudioWidgets/Private/AudioVectorscopeUMG.cpp
+++ b/Engine/Plugins/Runtime/AudioWidgets/Source/AudioWidgets/Private/AudioVectorscopeUMG.cpp
@@ -13,7 +13,7 @@
 
 UAudioVectorscope::UAudioVectorscope(const FObjectInitializer& ObjectInitializer)
 	: Super(ObjectInitializer)
-	, VectorscopeStyle(FAudioVectorscopePanelStyle::GetDefault())
+	, VectorscopeStyle(FAudioWidgetsStyle::Get().GetWidgetStyle<FAudioVectorscopePanelStyle>("AudioVectorscope.PanelStyle"))
 {
 #if WITH_EDITORONLY_DATA
 	AccessibleBehavior       = ESlateAccessibleBehavior::NotAccessible;
diff --git a/Engine/Plugins/Runtime/AudioWidgets/Source/AudioWidgets/Private/AudioWidgetsStyle.cpp b/Engine/Plugins/Runtime/AudioWidgets/Source/AudioWidgets/Private/AudioWidgetsStyle.cpp
--- a/Engine/Plugins/Runtime/AudioWidgets/Source/AudioWidgets/Private/AudioWidgetsStyle.cpp
+++ b/Engine/Plugins/Runtime/AudioWidgets/Source/AudioWidgets/Private/AudioWidgetsStyle.cpp
@@ -1,6 +1,11 @@
 // Copyright Epic Games, Inc. All Rights Reserved.
 
 #include "AudioWidgetsStyle.h"
+
+#include "AudioMeter.h"
+#include "AudioOscilloscopePanelStyle.h"
+#include "AudioSpectrumPlotStyle.h"
+#include "AudioVectorscopePanelStyle.h"
 #include "AudioWidgetsSlateTypes.h"
 #include "AudioMaterialSlate/AudioMaterialSlateTypes.h"
 #include "Brushes/SlateBoxBrush.h"
@@ -154,6 +159,17 @@ FAudioWidgetsStyle::FAudioWidgetsStyle()
 	Set("AudioMaterialSlider.Style", FAudioMaterialSliderStyle()
 		.SetTextBoxStyle(FAudioTextBoxStyle::GetDefault()));
 
+	/**
+	**Analyzer widget styles
+	*/
+	Set("AudioMeter.DefaultColorStyle", FAudioMeterDefaultColorStyle());
+	Set("AudioMeter.Style", FAudioMeterStyle());
+	Set("AudioOscilloscope.PanelStyle", FAudioOscilloscopePanelStyle());
+	Set("AudioSpectrumPlot.Style", FAudioSpectrumPlotStyle());
+	Set("AudioVectorscope.PanelStyle", FAudioVectorscopePanelStyle());
+	Set("SampledSequenceVectorViewer.Style", FSampledSequenceVectorViewerStyle());
+	Set("TriggerThresholdLine.Style", FTriggerThresholdLineStyle());
+
 	FSlateStyleRegistry::RegisterSlateStyle(*this);
 }
 void FAudioWidgetsStyle::SetResources()
diff --git a/Engine/Plugins/Runtime/AudioWidgets/Source/AudioWidgets/Public/AudioMeter.h b/Engine/Plugins/Runtime/AudioWidgets/Source/AudioWidgets/Public/AudioMeter.h
--- a/Engine/Plugins/Runtime/AudioWidgets/Source/AudioWidgets/Public/AudioMeter.h
+++ b/Engine/Plugins/Runtime/AudioWidgets/Source/AudioWidgets/Public/AudioMeter.h
@@ -35,6 +35,9 @@ struct AUDIOWIDGETS_API FAudioMeterDefaultColorStyle : public FSlateWidgetStyle
 	GENERATED_USTRUCT_BODY()
 
 public:
+	static const FName TypeName;
+	virtual const FName GetTypeName() const override { return TypeName; };
+
 	static const FAudioMeterDefaultColorStyle& GetDefault();
 
 	UPROPERTY(EditAnywhere, BlueprintReadWrite, Category = "Style")
diff --git a/Engine/Plugins/Runtime/AudioWidgets/Source/AudioWidgets/Public/SAudioMeter.h b/Engine/Plugins/Runtime/AudioWidgets/Source/AudioWidgets/Public/SAudioMeter.h
--- a/Engine/Plugins/Runtime/AudioWidgets/Source/AudioWidgets/Public/SAudioMeter.h
+++ b/Engine/Plugins/Runtime/AudioWidgets/Source/AudioWidgets/Public/SAudioMeter.h
@@ -2,6 +2,7 @@
 
 #pragma once
 
+#include "AudioWidgetsStyle.h"
 #include "Styling/ISlateStyle.h"
 #include "Styling/SlateWidgetStyleAsset.h"
 #include "Widgets/DeclarativeSyntaxSupport.h"
@@ -47,7 +48,7 @@ public:
 		, _MeterPeakColor(FLinearColor::Blue)
 		, _MeterScaleColor(FLinearColor::White)
 		, _MeterScaleLabelColor(FLinearColor::Gray)
-		, _Style(&FAudioMeterStyle::GetDefault())
+		, _Style(&FAudioWidgetsStyle::Get().GetWidgetStyle<FAudioMeterStyle>("AudioMeter.Style"))
 	{
 	}
 
diff --git a/Engine/Plugins/Runtime/AudioWidgets/Source/AudioWidgets/Public/SAudioOscilloscopePanelWidget.h b/Engine/Plugins/Runtime/AudioWidgets/Source/AudioWidgets/Public/SAudioOscilloscopePanelWidget.h
--- a/Engine/Plugins/Runtime/AudioWidgets/Source/AudioWidgets/Public/SAudioOscilloscopePanelWidget.h
+++ b/Engine/Plugins/Runtime/AudioWidgets/Source/AudioWidgets/Public/SAudioOscilloscopePanelWidget.h
@@ -56,7 +56,7 @@ public:
 		, _ValueGridMaxDivisionParameter(2)
 		, _SequenceRulerDisplayUnit(EXAxisLabelsUnit::Samples)
 		, _PanelLayoutType(EAudioPanelLayoutType::Basic)
-		, _PanelStyle(&FAudioOscilloscopePanelStyle::GetDefault())
+		, _PanelStyle(&FAudioWidgetsStyle::Get().GetWidgetStyle<FAudioOscilloscopePanelStyle>("AudioOscilloscope.PanelStyle"))
 	{}
 		
 		SLATE_DEFAULT_SLOT(FArguments, InArgs)
diff --git a/Engine/Plugins/Runtime/AudioWidgets/Source/AudioWidgets/Public/SAudioSpectrumPlot.h b/Engine/Plugins/Runtime/AudioWidgets/Source/AudioWidgets/Public/SAudioSpectrumPlot.h
--- a/Engine/Plugins/Runtime/AudioWidgets/Source/AudioWidgets/Public/SAudioSpectrumPlot.h
+++ b/Engine/Plugins/Runtime/AudioWidgets/Source/AudioWidgets/Public/SAudioSpectrumPlot.h
@@ -3,6 +3,7 @@
 #pragma once
 
 #include "AudioSpectrumPlotStyle.h"
+#include "AudioWidgetsStyle.h"
 #include "Framework/MultiBox/MultiBoxExtender.h"
 #include "Framework/SlateDelegates.h"
 #include "Styling/ISlateStyle.h"
@@ -130,7 +131,7 @@ class AUDIOWIDGETS_API SAudioSpectrumPlot : public SCompoundWidget
 {
 public:
 	SLATE_BEGIN_ARGS(SAudioSpectrumPlot)
-		: _Style(&FAudioSpectrumPlotStyle::GetDefault())
+		: _Style(&FAudioWidgetsStyle::Get().GetWidgetStyle<FAudioSpectrumPlotStyle>("AudioSpectrumPlot.Style"))
 		, _ViewMinFrequency(20.0f)
 		, _ViewMaxFrequency(20000.0f)
 		, _ViewMinSoundLevel(-60.0f)
diff --git a/Engine/Plugins/Runtime/AudioWidgets/Source/AudioWidgets/Public/SAudioVectorscopePanelWidget.h b/Engine/Plugins/Runtime/AudioWidgets/Source/AudioWidgets/Public/SAudioVectorscopePanelWidget.h
--- a/Engine/Plugins/Runtime/AudioWidgets/Source/AudioWidgets/Public/SAudioVectorscopePanelWidget.h
+++ b/Engine/Plugins/Runtime/AudioWidgets/Source/AudioWidgets/Public/SAudioVectorscopePanelWidget.h
@@ -4,6 +4,7 @@
 
 #include "AudioVectorscopePanelStyle.h"
 #include "AudioWidgetsEnums.h"
+#include "AudioWidgetsStyle.h"
 #include "Brushes/SlateColorBrush.h"
 #include "IFixedSampledSequenceViewReceiver.h"
 #include "SampledSequenceDrawingUtils.h"
@@ -24,7 +25,7 @@ public:
 		: _HideGrid(false)
 		, _ValueGridMaxDivisionParameter(2)
 		, _PanelLayoutType(EAudioPanelLayoutType::Basic)
-		, _PanelStyle(&FAudioVectorscopePanelStyle::GetDefault())
+		, _PanelStyle(&FAudioWidgetsStyle::Get().GetWidgetStyle<FAudioVectorscopePanelStyle>("AudioVectorscope.PanelStyle"))
 	{}
 		
 		SLATE_DEFAULT_SLOT(FArguments, InArgs)
diff --git a/Engine/Plugins/Runtime/AudioWidgets/Source/AudioWidgets/Public/STriggerThresholdLineWidget.h b/Engine/Plugins/Runtime/AudioWidgets/Source/AudioWidgets/Public/STriggerThresholdLineWidget.h
--- a/Engine/Plugins/Runtime/AudioWidgets/Source/AudioWidgets/Public/STriggerThresholdLineWidget.h
+++ b/Engine/Plugins/Runtime/AudioWidgets/Source/AudioWidgets/Public/STriggerThresholdLineWidget.h
@@ -2,6 +2,7 @@
 
 #pragma once
 
+#include "AudioWidgetsStyle.h"
 #include "Rendering/DrawElements.h"
 #include "SampledSequenceDrawingUtils.h"
 #include "Styling/ISlateStyle.h"
@@ -16,7 +17,7 @@ class AUDIOWIDGETS_API STriggerThresholdLineWidget : public SLeafWidget
 public:
     SLATE_BEGIN_ARGS(STriggerThresholdLineWidget)
         : _TriggerThreshold(0.0f)
-		, _Style(&FTriggerThresholdLineStyle::GetDefault())
+		, _Style(&FAudioWidgetsStyle::Get().GetWidgetStyle<FTriggerThresholdLineStyle>("TriggerThresholdLine.Style"))
     {}
     
 		SLATE_ARGUMENT(float, TriggerThreshold)
