AudioPropertiesDetailsInjector:
- allows to override the view of properties on UObject that want to implement a property sheet
- supports property matrix
- changes to AudioPropertiesEditorModule to allow the object to be instantiated on projects that are not aware of the plugin
- implementation into soundbase details
#rnx

diff --git a/Engine/Plugins/Runtime/AudioWidgets/Source/AudioWidgetsEditor/AudioWidgetsEditor.build.cs b/Engine/Plugins/Runtime/AudioWidgets/Source/AudioWidgetsEditor/AudioWidgetsEditor.build.cs
--- a/Engine/Plugins/Runtime/AudioWidgets/Source/AudioWidgetsEditor/AudioWidgetsEditor.build.cs
+++ b/Engine/Plugins/Runtime/AudioWidgets/Source/AudioWidgetsEditor/AudioWidgetsEditor.build.cs
@@ -8,13 +8,17 @@ public class AudioWidgetsEditor : ModuleRules
 	{
 		PCHUsage = ModuleRules.PCHUsageMode.UseExplicitOrSharedPCHs;
 
-		PrivateDependencyModuleNames.AddRange
+		PublicDependencyModuleNames.AddRange
 		(
-			new string[] 
+			new string[]
 			{
-				"Core",
-				"PropertyEditor",
+					"Core",
+					"CoreUObject",
+					"Engine",
+					"AudioExtensions",
+					"PropertyEditor"
 			}
 		);
+
 	}
 }
\ No newline at end of file
diff --git a/Engine/Plugins/Runtime/AudioWidgets/Source/AudioWidgetsEditor/Private/SoundBaseDetails.cpp b/Engine/Plugins/Runtime/AudioWidgets/Source/AudioWidgetsEditor/Private/SoundBaseDetails.cpp
--- a/Engine/Plugins/Runtime/AudioWidgets/Source/AudioWidgetsEditor/Private/SoundBaseDetails.cpp
+++ b/Engine/Plugins/Runtime/AudioWidgets/Source/AudioWidgetsEditor/Private/SoundBaseDetails.cpp
@@ -2,7 +2,9 @@
 
 #include "SoundBaseDetails.h"
 
+#include "IAudioPropertiesDetailsInjector.h"
 #include "DetailLayoutBuilder.h"
+#include "IPropertyUtilities.h"
 #include "Modules/ModuleManager.h"
 #include "PropertyHandle.h"
 #include "Sound/AudioSettings.h"
@@ -15,10 +17,69 @@ TSharedRef<IDetailCustomization> FSoundBaseDetails::MakeInstance()
 	return MakeShareable(new FSoundBaseDetails);
 }
 
+FSoundBaseDetails::~FSoundBaseDetails()
+{
+}
+
+void FSoundBaseDetails::PendingDelete()
+{
+	if (AudioPropertiesInjector.IsValid())
+	{
+		AudioPropertiesInjector->UnbindFromPropertySheetChanges();
+	}
+}
+
 void FSoundBaseDetails::CustomizeDetails(IDetailLayoutBuilder& DetailBuilder)
 {
-	if (!FModuleManager::Get().IsModuleLoaded("AudioProperties"))
+	if (!FModuleManager::Get().IsModuleLoaded("AudioProperties") || !FModuleManager::Get().IsModuleLoaded("AudioPropertiesEditor"))
 	{
 		DetailBuilder.HideCategory("AudioProperties");
 	}
+	else
+	{
+		InjectPropertySheetView(DetailBuilder);
+	}
 }
+
+void FSoundBaseDetails::InjectPropertySheetView(IDetailLayoutBuilder& DetailBuilder)
+{
+	check(FModuleManager::Get().IsModuleLoaded("AudioPropertiesEditor"))
+	
+	void* AudioPropertiesEditorModuleDllHandle = FPlatformProcess::GetDllHandle(*FModuleManager::Get().GetModuleFilename("AudioPropertiesEditor"));
+
+	if (!AudioPropertiesEditorModuleDllHandle)
+	{
+		return;
+	}
+
+	typedef IAudioPropertiesDetailsInjector* (*CreateInjectorFunc)();
+	CreateInjectorFunc CreateInjector = (CreateInjectorFunc)FPlatformProcess::GetDllExport(AudioPropertiesEditorModuleDllHandle, TEXT("CreateAudioPropertiesDetailsInjectorExtern"));
+
+	if (!CreateInjector)
+	{
+		return;
+	}
+
+	IAudioPropertiesDetailsInjector* RawInjector = CreateInjector();
+	if (RawInjector != nullptr)
+	{
+		AudioPropertiesInjector = TSharedPtr<IAudioPropertiesDetailsInjector>(RawInjector);
+	}
+	
+
+	if (!AudioPropertiesInjector.IsValid())
+	{
+		return;
+	}
+
+	TSharedRef<IPropertyHandle> PropertySheetHandle = DetailBuilder.GetProperty(GET_MEMBER_NAME_CHECKED(USoundBase, AudioPropertiesSheet));
+
+	if (!PropertySheetHandle->IsValidHandle())
+	{
+		UE_LOG(LogInit, Log, TEXT("Invalid Property Sheet Handle found when customizing SoundBase details, property sheet view will not be injected"))
+		return;
+	}
+
+	AudioPropertiesInjector->CustomizeInjectedPropertiesDetails(DetailBuilder, PropertySheetHandle);
+	AudioPropertiesInjector->BindDetailCustomizationToPropertySheetChanges(DetailBuilder, PropertySheetHandle);
+}
\ No newline at end of file
diff --git a/Engine/Plugins/Runtime/AudioWidgets/Source/AudioWidgetsEditor/Private/SoundBaseDetails.h b/Engine/Plugins/Runtime/AudioWidgets/Source/AudioWidgetsEditor/Private/SoundBaseDetails.h
--- a/Engine/Plugins/Runtime/AudioWidgets/Source/AudioWidgetsEditor/Private/SoundBaseDetails.h
+++ b/Engine/Plugins/Runtime/AudioWidgets/Source/AudioWidgetsEditor/Private/SoundBaseDetails.h
@@ -6,6 +6,7 @@
 #include "Templates/SharedPointer.h"
 
 class IDetailLayoutBuilder;
+class IAudioPropertiesDetailsInjector;
 
 class FSoundBaseDetails : public IDetailCustomization
 {
@@ -13,7 +14,15 @@ public:
 	/** Makes a new instance of this detail layout class for a specific detail view requesting it */
 	static TSharedRef<IDetailCustomization> MakeInstance();
 
+	virtual ~FSoundBaseDetails();
+
+	virtual void PendingDelete() override;
+
 private:
 	/** IDetailCustomization interface */
 	virtual void CustomizeDetails(IDetailLayoutBuilder& DetailBuilder) override;
+
+	void InjectPropertySheetView(IDetailLayoutBuilder& DetailBuilder);
+
+	TSharedPtr<IAudioPropertiesDetailsInjector> AudioPropertiesInjector;
 };
diff --git a/Engine/Source/Runtime/AudioExtensions/Private/IAudioPropertiesDetailsInjector.cpp b/Engine/Source/Runtime/AudioExtensions/Private/IAudioPropertiesDetailsInjector.cpp
new file mode 100644
--- /dev/null
+++ b/Engine/Source/Runtime/AudioExtensions/Private/IAudioPropertiesDetailsInjector.cpp
@@ -0,0 +1,8 @@
+// Copyright Epic Games, Inc. All Rights Reserved.
+
+#pragma once
+
+#include "IAudioPropertiesDetailsInjector.h"
+
+
+IAudioPropertiesDetailsInjector::~IAudioPropertiesDetailsInjector() {};
\ No newline at end of file
diff --git a/Engine/Source/Runtime/AudioExtensions/Public/IAudioPropertiesDetailsInjector.h b/Engine/Source/Runtime/AudioExtensions/Public/IAudioPropertiesDetailsInjector.h
new file mode 100644
--- /dev/null
+++ b/Engine/Source/Runtime/AudioExtensions/Public/IAudioPropertiesDetailsInjector.h
@@ -0,0 +1,18 @@
+// Copyright Epic Games, Inc. All Rights Reserved.
+
+#pragma once
+
+#include "Templates/SharedPointer.h"
+
+class IDetailLayoutBuilder;
+class IPropertyHandle;
+
+class AUDIOEXTENSIONS_API IAudioPropertiesDetailsInjector : public TSharedFromThis<IAudioPropertiesDetailsInjector>
+{
+public:
+	virtual ~IAudioPropertiesDetailsInjector() = 0;
+
+	virtual void CustomizeInjectedPropertiesDetails(IDetailLayoutBuilder& DetailBuilder, TSharedRef<IPropertyHandle> PropertySheetPropertyHandle) = 0;
+	virtual void BindDetailCustomizationToPropertySheetChanges(IDetailLayoutBuilder& DetailBuilder, TSharedRef<IPropertyHandle> PropertySheetPropertyHandle) = 0;
+	virtual void UnbindFromPropertySheetChanges() = 0;
+};
\ No newline at end of file
