[AudioProperties ]fix for cis error  620430
#rnx

diff --git a/Engine/Plugins/Runtime/AudioWidgets/Source/AudioWidgetsEditor/Private/SoundBaseDetails.cpp b/Engine/Plugins/Runtime/AudioWidgets/Source/AudioWidgetsEditor/Private/SoundBaseDetails.cpp
--- a/Engine/Plugins/Runtime/AudioWidgets/Source/AudioWidgetsEditor/Private/SoundBaseDetails.cpp
+++ b/Engine/Plugins/Runtime/AudioWidgets/Source/AudioWidgetsEditor/Private/SoundBaseDetails.cpp
@@ -44,8 +44,14 @@ void FSoundBaseDetails::CustomizeDetails(IDetailLayoutBuilder& DetailBuilder)
 void FSoundBaseDetails::InjectPropertySheetView(IDetailLayoutBuilder& DetailBuilder)
 {
 	check(FModuleManager::Get().IsModuleLoaded("AudioPropertiesEditor"))
-	
-	void* AudioPropertiesEditorModuleDllHandle = FPlatformProcess::GetDllHandle(*FModuleManager::Get().GetModuleFilename("AudioPropertiesEditor"));
+
+	FString ModuleFileName = "UnrealEditor-AudioPropertiesEditor.dll";
+
+#if !IS_MONOLITHIC
+	ModuleFileName = FModuleManager::Get().GetModuleFilename("AudioPropertiesEditor");
+#endif
+
+	void* AudioPropertiesEditorModuleDllHandle = FPlatformProcess::GetDllHandle(*ModuleFileName);
 
 	if (!AudioPropertiesEditorModuleDllHandle)
 	{
