Fix mac build.

diff --git a/Engine/Plugins/Runtime/AudioWidgets/Source/AudioWidgets/Private/AudioAnalyzerRack.cpp b/Engine/Plugins/Runtime/AudioWidgets/Source/AudioWidgets/Private/AudioAnalyzerRack.cpp
--- a/Engine/Plugins/Runtime/AudioWidgets/Source/AudioWidgets/Private/AudioAnalyzerRack.cpp
+++ b/Engine/Plugins/Runtime/AudioWidgets/Source/AudioWidgets/Private/AudioAnalyzerRack.cpp
@@ -187,25 +187,25 @@ namespace AudioWidgets
 						FUIAction(
 							FExecuteAction::CreateLambda([InTabManager, TabId]()
 								{
-									if (TSharedPtr<FTabManager> TabManager = InTabManager.Pin())
+									if (TSharedPtr<FTabManager> PinnedTabManager = InTabManager.Pin())
 									{
-										TSharedPtr<SDockTab> LiveTab = TabManager->FindExistingLiveTab(TabId);
+										TSharedPtr<SDockTab> LiveTab = PinnedTabManager->FindExistingLiveTab(TabId);
 										if (!LiveTab.IsValid())
 										{
-											TabManager->TryInvokeTab(TabId);
+											PinnedTabManager->TryInvokeTab(TabId);
 										}
 										else
 										{
 											LiveTab->RequestCloseTab();
 										}
-										TabManager->SavePersistentLayout();
+										PinnedTabManager->SavePersistentLayout();
 									}
 								}),
 							FCanExecuteAction(),
 							FIsActionChecked::CreateLambda([InTabManager, TabId]()
 								{
-									TSharedPtr<FTabManager> TabManager = InTabManager.Pin();
-									return TabManager.IsValid() && TabManager->FindExistingLiveTab(TabId).IsValid();
+									TSharedPtr<FTabManager> PinnedTabManager = InTabManager.Pin();
+									return PinnedTabManager.IsValid() && PinnedTabManager->FindExistingLiveTab(TabId).IsValid();
 								})
 						),
 						NAME_None,
