Feedback from code review - Make IAudioAnalyzerRackUnit inherit from TSharedFromThis instead of FAudioSpectrumAnalyzer.

diff --git a/Engine/Plugins/Runtime/AudioWidgets/Source/AudioWidgets/Public/AudioAnalyzerRack.h b/Engine/Plugins/Runtime/AudioWidgets/Source/AudioWidgets/Public/AudioAnalyzerRack.h
--- a/Engine/Plugins/Runtime/AudioWidgets/Source/AudioWidgets/Public/AudioAnalyzerRack.h
+++ b/Engine/Plugins/Runtime/AudioWidgets/Source/AudioWidgets/Public/AudioAnalyzerRack.h
@@ -21,7 +21,7 @@ namespace AudioWidgets
 	/**
 	 * Interface for something that can be used in an analyzer rack.
 	 */
-	class IAudioAnalyzerRackUnit
+	class IAudioAnalyzerRackUnit : public TSharedFromThis<IAudioAnalyzerRackUnit>
 	{
 	public:
 		virtual ~IAudioAnalyzerRackUnit() {};
diff --git a/Engine/Plugins/Runtime/AudioWidgets/Source/AudioWidgets/Public/AudioSpectrumAnalyzer.h b/Engine/Plugins/Runtime/AudioWidgets/Source/AudioWidgets/Public/AudioSpectrumAnalyzer.h
--- a/Engine/Plugins/Runtime/AudioWidgets/Source/AudioWidgets/Public/AudioSpectrumAnalyzer.h
+++ b/Engine/Plugins/Runtime/AudioWidgets/Source/AudioWidgets/Public/AudioSpectrumAnalyzer.h
@@ -69,7 +69,7 @@ namespace AudioWidgets
 	 * Exponential time-smoothing is applied to the spectrum.
 	 * Can either create an Audio Bus to analyze, or analyze the given Bus.
 	 */
-	class AUDIOWIDGETS_API FAudioSpectrumAnalyzer : public IAudioAnalyzerRackUnit, public TSharedFromThis<FAudioSpectrumAnalyzer>
+	class AUDIOWIDGETS_API FAudioSpectrumAnalyzer : public IAudioAnalyzerRackUnit
 	{
 	public:
 		static const FAudioAnalyzerRackUnitTypeInfo RackUnitTypeInfo;
