[AudioWidgets] Added time-frequency spectrogram audio analyzer rack unit.

diff --git a/Engine/Plugins/Runtime/AudioWidgets/Source/AudioWidgets/Private/AudioAnalyzerRack.cpp b/Engine/Plugins/Runtime/AudioWidgets/Source/AudioWidgets/Private/AudioAnalyzerRack.cpp
--- a/Engine/Plugins/Runtime/AudioWidgets/Source/AudioWidgets/Private/AudioAnalyzerRack.cpp
+++ b/Engine/Plugins/Runtime/AudioWidgets/Source/AudioWidgets/Private/AudioAnalyzerRack.cpp
@@ -24,6 +24,24 @@ namespace AudioWidgets
 			TabManagerLayoutName = TEXT("AudioWidgets_FAudioAnalyzerRack_v0");
 		}
 
+		{
+			// Append rack unit type names to the tab manager layout name.
+			// When a new rack unit type becomes available a new default layout will be generated including the new rack unit.
+			FStringBuilderBase StringBuilder;
+			StringBuilder.Append(TabManagerLayoutName.ToString());
+
+			TArray<FName> RackUnitTypeNames;
+			FAudioAnalyzerRackUnitRegistry::Get().GetRegisteredRackUnitTypeNames(RackUnitTypeNames);
+			RackUnitTypeNames.Sort([](const FName& A, const FName& B) { return A.Compare(B) < 0; });
+			for (FName RackUnitTypeName : RackUnitTypeNames)
+			{
+				StringBuilder.AppendChar('_');
+				StringBuilder.Append(RackUnitTypeName.ToString());
+			}
+
+			TabManagerLayoutName = StringBuilder.ToString();
+		}
+
 		RackUnitConstructParams.StyleSet = (Params.StyleSet != nullptr) ? Params.StyleSet : &FAudioWidgetsStyle::Get();
 		RackUnitConstructParams.EditorSettingsClass = Params.EditorSettingsClass;
 	}
diff --git a/Engine/Plugins/Runtime/AudioWidgets/Source/AudioWidgets/Private/AudioAnalyzerRackUnitRegistry.cpp b/Engine/Plugins/Runtime/AudioWidgets/Source/AudioWidgets/Private/AudioAnalyzerRackUnitRegistry.cpp
--- a/Engine/Plugins/Runtime/AudioWidgets/Source/AudioWidgets/Private/AudioAnalyzerRackUnitRegistry.cpp
+++ b/Engine/Plugins/Runtime/AudioWidgets/Source/AudioWidgets/Private/AudioAnalyzerRackUnitRegistry.cpp
@@ -31,4 +31,9 @@ namespace AudioWidgets
 	{
 		RegisteredRackUnitTypes.GenerateValueArray(OutArray);
 	}
+
+	void FAudioAnalyzerRackUnitRegistry::GetRegisteredRackUnitTypeNames(TArray<FName>& OutArray) const
+	{
+		RegisteredRackUnitTypes.GenerateKeyArray(OutArray);
+	}
 }
diff --git a/Engine/Plugins/Runtime/AudioWidgets/Source/AudioWidgets/Private/AudioAnalyzerRackUnitRegistry.h b/Engine/Plugins/Runtime/AudioWidgets/Source/AudioWidgets/Private/AudioAnalyzerRackUnitRegistry.h
--- a/Engine/Plugins/Runtime/AudioWidgets/Source/AudioWidgets/Private/AudioAnalyzerRackUnitRegistry.h
+++ b/Engine/Plugins/Runtime/AudioWidgets/Source/AudioWidgets/Private/AudioAnalyzerRackUnitRegistry.h
@@ -16,6 +16,7 @@ namespace AudioWidgets
 		TSharedRef<IAudioAnalyzerRackUnit> MakeRackUnit(FName RackUnitTypeName, const FAudioAnalyzerRackUnitConstructParams& Params);
 
 		void GetRegisteredRackUnitTypes(TArray<const FAudioAnalyzerRackUnitTypeInfo*>& OutArray) const;
+		void GetRegisteredRackUnitTypeNames(TArray<FName>& OutArray) const;
 
 	private:
 		TMap<FName, const FAudioAnalyzerRackUnitTypeInfo*> RegisteredRackUnitTypes;
diff --git a/Engine/Plugins/Runtime/AudioWidgets/Source/AudioWidgets/Private/AudioSpectrogram.cpp b/Engine/Plugins/Runtime/AudioWidgets/Source/AudioWidgets/Private/AudioSpectrogram.cpp
new file mode 100644
--- /dev/null
+++ b/Engine/Plugins/Runtime/AudioWidgets/Source/AudioWidgets/Private/AudioSpectrogram.cpp
@@ -0,0 +1,534 @@
+// Copyright Epic Games, Inc. All Rights Reserved.
+
+#include "AudioSpectrogram.h"
+
+#include "AudioDevice.h"
+#include "AudioDeviceManager.h"
+#include "Framework/MultiBox/MultiBoxBuilder.h"
+#include "Widgets/Docking/SDockTab.h"
+
+#define LOCTEXT_NAMESPACE "FAudioSpectrogram"
+
+namespace AudioWidgets
+{
+	namespace AudioSpectrogramPrivate
+	{
+		/**
+		 * Light wrapper for accessing settings for the analyzer rack unit. Can be passed by value.
+		 */
+		class FRackUnitSettingsHelper
+		{
+		public:
+			FRackUnitSettingsHelper(const FProperty& InSettingsProperty)
+				: SettingsProperty(InSettingsProperty)
+			{
+				//
+			}
+
+			FSpectrogramRackUnitSettings* GetRackUnitSettings() const
+			{
+				UObject* EditorSettingsObject = GetEditorSettingsObject();
+				return SettingsProperty.ContainerPtrToValuePtr<FSpectrogramRackUnitSettings>(EditorSettingsObject);
+			}
+
+			void SaveConfig() const
+			{
+				GetEditorSettingsObject()->SaveConfig();
+			}
+
+		private:
+			UObject* GetEditorSettingsObject() const
+			{
+				return SettingsProperty.GetOwnerClass()->GetDefaultObject();
+			}
+
+			const FProperty& SettingsProperty;
+		};
+	}
+
+	const FAudioAnalyzerRackUnitTypeInfo FAudioSpectrogram::RackUnitTypeInfo
+	{
+		.TypeName = TEXT("FAudioSpectrogram"),
+		.DisplayName = LOCTEXT("AudioSpectrogramDisplayName", "Spectrogram"),
+		.OnMakeAudioAnalyzerRackUnit = FOnMakeAudioAnalyzerRackUnit::CreateStatic(&MakeRackUnit),
+		.VerticalSizeCoefficient = 0.25f,
+	};
+
+	FAudioSpectrogram::FAudioSpectrogram(const FAudioSpectrogramParams& Params)
+		: SpectrumAnalysisSettings(NewObject<USynesthesiaSpectrumAnalysisSettings>())
+		, ConstantQSettings(NewObject<UConstantQSettings>())
+		, Widget(SNew(SAudioSpectrogram)
+			.Clipping(EWidgetClipping::ClipToBounds)
+			.FrequencyAxisScale(Params.FrequencyAxisScale)
+			.FrequencyAxisPixelBucketMode(Params.FrequencyAxisPixelBucketMode)
+			.ColorMap(Params.ColorMap)
+			.Orientation(Params.Orientation)
+			.OnFrequencyAxisPixelBucketModeMenuEntryClicked(Params.OnFrequencyAxisPixelBucketModeMenuEntryClicked)
+			.OnFrequencyAxisScaleMenuEntryClicked(Params.OnFrequencyAxisScaleMenuEntryClicked)
+			.OnColorMapMenuEntryClicked(Params.OnColorMapMenuEntryClicked)
+			.OnOrientationMenuEntryClicked(Params.OnOrientationMenuEntryClicked)
+		)
+		, AnalyzerType(Params.AnalyzerType)
+		, FFTAnalyzerFFTSize(Params.FFTAnalyzerFFTSize)
+		, CQTAnalyzerFFTSize(Params.CQTAnalyzerFFTSize)
+		, OnAnalyzerTypeMenuEntryClicked(Params.OnAnalyzerTypeMenuEntryClicked)
+		, OnFFTAnalyzerFFTSizeMenuEntryClicked(Params.OnFFTAnalyzerFFTSizeMenuEntryClicked)
+		, OnCQTAnalyzerFFTSizeMenuEntryClicked(Params.OnCQTAnalyzerFFTSizeMenuEntryClicked)
+	{
+		SpectrumAnalysisSettings->SpectrumType = EAudioSpectrumType::PowerSpectrum;
+		SpectrumAnalysisSettings->FFTSize = FFTAnalyzerFFTSize.Get();
+		SpectrumAnalysisSettings->WindowType = EFFTWindowType::Blackman;
+		SpectrumAnalysisSettings->bDownmixToMono = true;
+
+		ConstantQSettings->SpectrumType = EAudioSpectrumType::PowerSpectrum;
+		ConstantQSettings->NumBandsPerOctave = 6.0f;
+		ConstantQSettings->NumBands = 61;
+		ConstantQSettings->StartingFrequencyHz = 20000.0f * FMath::Pow(0.5f, (ConstantQSettings->NumBands - 1) / ConstantQSettings->NumBandsPerOctave);
+		ConstantQSettings->FFTSize = CQTAnalyzerFFTSize.Get();
+		ConstantQSettings->WindowType = EFFTWindowType::Blackman;
+		ConstantQSettings->bDownmixToMono = true;
+		ConstantQSettings->BandWidthStretch = 2.0f;
+
+		ContextMenuExtension = Widget->AddContextMenuExtension(EExtensionHook::Before, nullptr, FMenuExtensionDelegate::CreateRaw(this, &FAudioSpectrogram::ExtendSpectrumPlotContextMenu));
+
+		ActiveTimer = Widget->RegisterActiveTimer(0.0f, FWidgetActiveTimerDelegate::CreateRaw(this, &FAudioSpectrogram::Update));
+
+		Init(Params.NumChannels, Params.AudioDeviceId, Params.ExternalAudioBus);
+	}
+
+	FAudioSpectrogram::~FAudioSpectrogram()
+	{
+		Teardown();
+
+		if (ContextMenuExtension.IsValid())
+		{
+			Widget->RemoveContextMenuExtension(ContextMenuExtension.ToSharedRef());
+		}
+
+		if (ActiveTimer.IsValid())
+		{
+			Widget->UnRegisterActiveTimer(ActiveTimer.ToSharedRef());
+		}
+	}
+
+	UAudioBus* FAudioSpectrogram::GetAudioBus() const
+	{
+		return AudioBus.Get();
+	}
+
+	TSharedRef<SWidget> FAudioSpectrogram::GetWidget() const
+	{
+		return Widget->AsShared();
+	}
+
+	void FAudioSpectrogram::Init(int32 InNumChannels, Audio::FDeviceId InAudioDeviceId, TObjectPtr<UAudioBus> InExternalAudioBus)
+	{
+		check(InNumChannels > 0);
+
+		Teardown();
+
+		AudioDeviceId = InAudioDeviceId;
+		bUseExternalAudioBus = InExternalAudioBus != nullptr;
+		AudioBus = bUseExternalAudioBus ? TStrongObjectPtr(InExternalAudioBus.Get()) : TStrongObjectPtr(NewObject<UAudioBus>());
+		AudioBus->AudioBusChannels = EAudioBusChannels(InNumChannels - 1);
+
+		CreateSynesthesiaSpectrumAnalyzer();
+		CreateConstantQAnalyzer();
+
+		StartAnalyzing(AnalyzerType.Get());
+	}
+
+	void FAudioSpectrogram::StartAnalyzing(const EAudioSpectrumAnalyzerType InAnalyzerType)
+	{
+		ensure(!ActiveAnalyzerType.IsSet());
+
+		switch (InAnalyzerType)
+		{
+		case EAudioSpectrumAnalyzerType::FFT:
+			SpectrumAnalyzer->StartAnalyzing(AudioDeviceId, AudioBus.Get());
+			break;
+		case EAudioSpectrumAnalyzerType::CQT:
+			ConstantQAnalyzer->StartAnalyzing(AudioDeviceId, AudioBus.Get());
+			break;
+		default:
+			break;
+		}
+
+		ActiveAnalyzerType = InAnalyzerType;
+	}
+
+	void FAudioSpectrogram::StopAnalyzing()
+	{
+		ensure(ActiveAnalyzerType.IsSet());
+
+		switch (ActiveAnalyzerType.GetValue())
+		{
+		case EAudioSpectrumAnalyzerType::FFT:
+			SpectrumAnalyzer->StopAnalyzing();
+			break;
+		case EAudioSpectrumAnalyzerType::CQT:
+			ConstantQAnalyzer->StopAnalyzing();
+			break;
+		default:
+			break;
+		}
+
+		ActiveAnalyzerType.Reset();
+	}
+
+	void FAudioSpectrogram::OnSpectrumResults(USynesthesiaSpectrumAnalyzer* InSpectrumAnalyzer, int32 ChannelIndex, const TArray<FSynesthesiaSpectrumResults>& InSpectrumResultsArray)
+	{
+		if (ActiveAnalyzerType == EAudioSpectrumAnalyzerType::FFT && InSpectrumAnalyzer == SpectrumAnalyzer.Get())
+		{
+			for (const FSynesthesiaSpectrumResults& SpectrumResults : InSpectrumResultsArray)
+			{
+				// Find samplerate:
+				float SampleRate = 48000.0f;
+				if (const FAudioDeviceManager* AudioDeviceManager = FAudioDeviceManager::Get())
+				{
+					if (const FAudioDevice* AudioDevice = AudioDeviceManager->GetAudioDeviceRaw(AudioDeviceId))
+					{
+						SampleRate = AudioDevice->GetSampleRate();
+					}
+				}
+
+				Widget->AddFrame(SpectrumResults, SpectrumAnalysisSettings->SpectrumType, SampleRate);
+			}
+		}
+	}
+
+	void FAudioSpectrogram::OnConstantQResults(UConstantQAnalyzer* InSpectrumAnalyzer, int32 ChannelIndex, const TArray<FConstantQResults>& InSpectrumResultsArray)
+	{
+		if (ActiveAnalyzerType == EAudioSpectrumAnalyzerType::CQT && InSpectrumAnalyzer == ConstantQAnalyzer.Get())
+		{
+			for (const FConstantQResults& SpectrumResults : InSpectrumResultsArray)
+			{
+				Widget->AddFrame(SpectrumResults, ConstantQSettings->StartingFrequencyHz, ConstantQSettings->NumBandsPerOctave, ConstantQSettings->SpectrumType);
+			}
+		}
+	}
+
+	void FAudioSpectrogram::Teardown()
+	{
+		if (SpectrumAnalyzer.IsValid() && SpectrumAnalyzer->IsValidLowLevel())
+		{
+			if (ActiveAnalyzerType == EAudioSpectrumAnalyzerType::FFT)
+			{
+				SpectrumAnalyzer->StopAnalyzing();
+			}
+
+			ReleaseSynesthesiaSpectrumAnalyzer();
+		}
+
+		if (ConstantQAnalyzer.IsValid() && ConstantQAnalyzer->IsValidLowLevel())
+		{
+			if (ActiveAnalyzerType == EAudioSpectrumAnalyzerType::CQT)
+			{
+				ConstantQAnalyzer->StopAnalyzing();
+			}
+
+			ReleaseConstantQAnalyzer();
+		}
+
+		ActiveAnalyzerType.Reset();
+
+		AudioBus.Reset();
+		bUseExternalAudioBus = false;
+	}
+
+	void FAudioSpectrogram::ExtendSpectrumPlotContextMenu(FMenuBuilder& MenuBuilder)
+	{
+		MenuBuilder.BeginSection("AnalyzerSettings", LOCTEXT("AnalyzerSettings", "Analyzer Settings"));
+		MenuBuilder.AddSubMenu(
+			LOCTEXT("AnalyzerType", "Analyzer Type"),
+			FText(),
+			FNewMenuDelegate::CreateSP(this, &FAudioSpectrogram::BuildAnalyzerTypeSubMenu));
+		MenuBuilder.AddSubMenu(
+			LOCTEXT("FFTSize", "FFT Size"),
+			FText(),
+			FNewMenuDelegate::CreateSP(this, &FAudioSpectrogram::BuildFFTSizeSubMenu));
+		MenuBuilder.EndSection();
+	}
+
+	void FAudioSpectrogram::BuildAnalyzerTypeSubMenu(FMenuBuilder& SubMenu)
+	{
+		const UEnum* EnumClass = StaticEnum<EAudioSpectrumAnalyzerType>();
+		const int32 NumEnumValues = EnumClass->NumEnums() - 1; // Exclude 'MAX' enum value.
+		for (int32 Index = 0; Index < NumEnumValues; Index++)
+		{
+			const auto EnumValue = static_cast<EAudioSpectrumAnalyzerType>(EnumClass->GetValueByIndex(Index));
+
+			SubMenu.AddMenuEntry(
+				EnumClass->GetDisplayNameTextByIndex(Index),
+#if WITH_EDITOR
+				EnumClass->GetToolTipTextByIndex(Index),
+#else
+				FText(),
+#endif
+				FSlateIcon(),
+				FUIAction(
+					FExecuteAction::CreateSPLambda(this, [this, EnumValue]()
+						{
+							if (!AnalyzerType.IsBound())
+							{
+								AnalyzerType = EnumValue;
+							}
+
+							OnAnalyzerTypeMenuEntryClicked.ExecuteIfBound(EnumValue);
+						}),
+					FCanExecuteAction(),
+					FIsActionChecked::CreateSPLambda(this, [this, EnumValue]() { return (AnalyzerType.Get() == EnumValue); })
+				),
+				NAME_None,
+				EUserInterfaceActionType::ToggleButton);
+		}
+	}
+
+	void FAudioSpectrogram::BuildFFTSizeSubMenu(FMenuBuilder& SubMenu)
+	{
+		// There is a different FFTSize enum depending on the analyzer type.
+
+		if (AnalyzerType.Get() == EAudioSpectrumAnalyzerType::FFT)
+		{
+			const UEnum* EnumClass = StaticEnum<EFFTSize>();
+			const int32 NumEnumValues = EnumClass->NumEnums() - 1; // Exclude 'MAX' enum value.
+			for (int32 Index = 0; Index < NumEnumValues; Index++)
+			{
+				const auto EnumValue = static_cast<EFFTSize>(EnumClass->GetValueByIndex(Index));
+				if (EnumValue == EFFTSize::DefaultSize)
+				{
+					// Skip the duplicate 512 enum value 'DefaultSize'.
+					continue;
+				}
+
+				SubMenu.AddMenuEntry(
+					EnumClass->GetDisplayNameTextByIndex(Index),
+#if WITH_EDITOR
+					EnumClass->GetToolTipTextByIndex(Index),
+#else
+					FText(),
+#endif
+					FSlateIcon(),
+					FUIAction(
+						FExecuteAction::CreateSPLambda(this, [this, EnumValue]()
+							{
+								if (!FFTAnalyzerFFTSize.IsBound())
+								{
+									FFTAnalyzerFFTSize = EnumValue;
+								}
+
+								OnFFTAnalyzerFFTSizeMenuEntryClicked.ExecuteIfBound(EnumValue);
+							}),
+						FCanExecuteAction(),
+						FIsActionChecked::CreateSPLambda(this, [this, EnumValue]() { return (FFTAnalyzerFFTSize.Get() == EnumValue); })
+					),
+					NAME_None,
+					EUserInterfaceActionType::ToggleButton);
+			}
+		}
+		else if (AnalyzerType.Get() == EAudioSpectrumAnalyzerType::CQT)
+		{
+			const UEnum* EnumClass = StaticEnum<EConstantQFFTSizeEnum>();
+			const int32 NumEnumValues = EnumClass->NumEnums() - 1; // Exclude 'MAX' enum value.
+			for (int32 Index = 0; Index < NumEnumValues; Index++)
+			{
+				const auto EnumValue = static_cast<EConstantQFFTSizeEnum>(EnumClass->GetValueByIndex(Index));
+
+				SubMenu.AddMenuEntry(
+					EnumClass->GetDisplayNameTextByIndex(Index),
+#if WITH_EDITOR
+					EnumClass->GetToolTipTextByIndex(Index),
+#else
+					FText(),
+#endif
+					FSlateIcon(),
+					FUIAction(
+						FExecuteAction::CreateSPLambda(this, [this, EnumValue]()
+							{
+								if (!CQTAnalyzerFFTSize.IsBound())
+								{
+									CQTAnalyzerFFTSize = EnumValue;
+								}
+
+								OnCQTAnalyzerFFTSizeMenuEntryClicked.ExecuteIfBound(EnumValue);
+							}),
+						FCanExecuteAction(),
+						FIsActionChecked::CreateSPLambda(this, [this, EnumValue]() { return (CQTAnalyzerFFTSize.Get() == EnumValue); })
+					),
+					NAME_None,
+					EUserInterfaceActionType::ToggleButton);
+			}
+		}
+	}
+
+	EActiveTimerReturnType FAudioSpectrogram::Update(double InCurrentTime, float InDeltaTime)
+	{
+		const EAudioSpectrumAnalyzerType RequiredAnalyzerType = AnalyzerType.Get();
+		const EFFTSize FFTAnalyzerRequiredFFTSize = FFTAnalyzerFFTSize.Get();
+		const EConstantQFFTSizeEnum CQTAnalyzerRequiredFFTSize = CQTAnalyzerFFTSize.Get();
+
+		const bool bRequiredAnalyzerTypeChanged = (ActiveAnalyzerType != RequiredAnalyzerType);
+		const bool bFFTAnalyzerRequiredFFTSizeChanged = (SpectrumAnalysisSettings->FFTSize != FFTAnalyzerRequiredFFTSize);
+		const bool bCQTAnalyzerRequiredFFTSizeChanged = (ConstantQSettings->FFTSize != CQTAnalyzerRequiredFFTSize);
+		if (bRequiredAnalyzerTypeChanged || bFFTAnalyzerRequiredFFTSizeChanged || bCQTAnalyzerRequiredFFTSizeChanged)
+		{
+			StopAnalyzing();
+
+			if (bFFTAnalyzerRequiredFFTSizeChanged)
+			{
+				ReleaseSynesthesiaSpectrumAnalyzer();
+				SpectrumAnalysisSettings->FFTSize = FFTAnalyzerRequiredFFTSize;
+				CreateSynesthesiaSpectrumAnalyzer();
+			}
+
+			if (bCQTAnalyzerRequiredFFTSizeChanged)
+			{
+				ReleaseConstantQAnalyzer();
+				ConstantQSettings->FFTSize = CQTAnalyzerRequiredFFTSize;
+				CreateConstantQAnalyzer();
+			}
+
+			StartAnalyzing(RequiredAnalyzerType);
+		}
+
+		return EActiveTimerReturnType::Continue;
+	}
+
+	void FAudioSpectrogram::CreateSynesthesiaSpectrumAnalyzer()
+	{
+		ensure(!SpectrumAnalyzer.IsValid());
+		ensure(!SpectrumResultsDelegateHandle.IsValid());
+
+		SpectrumAnalyzer = TStrongObjectPtr(NewObject<USynesthesiaSpectrumAnalyzer>());
+		SpectrumAnalyzer->Settings = SpectrumAnalysisSettings.Get();
+		SpectrumResultsDelegateHandle = SpectrumAnalyzer->OnSpectrumResultsNative.AddRaw(this, &FAudioSpectrogram::OnSpectrumResults);
+	}
+
+	void FAudioSpectrogram::ReleaseSynesthesiaSpectrumAnalyzer()
+	{
+		if (ensure(SpectrumAnalyzer.IsValid() && SpectrumResultsDelegateHandle.IsValid()))
+		{
+			SpectrumAnalyzer->OnSpectrumResultsNative.Remove(SpectrumResultsDelegateHandle);
+		}
+
+		SpectrumResultsDelegateHandle.Reset();
+		SpectrumAnalyzer.Reset();
+	}
+
+	void FAudioSpectrogram::CreateConstantQAnalyzer()
+	{
+		ConstantQAnalyzer = TStrongObjectPtr(NewObject<UConstantQAnalyzer>());
+		ConstantQAnalyzer->Settings = ConstantQSettings.Get();
+		ConstantQResultsDelegateHandle = ConstantQAnalyzer->OnConstantQResultsNative.AddRaw(this, &FAudioSpectrogram::OnConstantQResults);
+	}
+
+	void FAudioSpectrogram::ReleaseConstantQAnalyzer()
+	{
+		if (ensure(ConstantQAnalyzer.IsValid() && ConstantQResultsDelegateHandle.IsValid()))
+		{
+			ConstantQAnalyzer->OnConstantQResultsNative.Remove(ConstantQResultsDelegateHandle);
+		}
+
+		ConstantQResultsDelegateHandle.Reset();
+		ConstantQAnalyzer.Reset();
+	}
+
+	TSharedRef<SDockTab> FAudioSpectrogram::SpawnTab(const FSpawnTabArgs& Args) const
+	{
+		return SNew(SDockTab)
+			.Clipping(EWidgetClipping::ClipToBounds)
+			.Label(RackUnitTypeInfo.DisplayName)
+			[
+				GetWidget()
+			];
+	}
+
+	TSharedRef<IAudioAnalyzerRackUnit> FAudioSpectrogram::MakeRackUnit(const FAudioAnalyzerRackUnitConstructParams& Params)
+	{
+		using namespace AudioSpectrogramPrivate;
+
+		FAudioSpectrogramParams AnalyzerParams;
+		AnalyzerParams.NumChannels = Params.AudioBusInfo.AudioBus->GetNumChannels();
+		AnalyzerParams.AudioDeviceId = Params.AudioBusInfo.AudioDeviceId;
+		AnalyzerParams.ExternalAudioBus = Params.AudioBusInfo.AudioBus;
+
+		if (Params.EditorSettingsClass != nullptr)
+		{
+			// If we have been given a valid editor settings class, bind analyzer options to the settings:
+			const FProperty* SpectrogramSettingsProperty = Params.EditorSettingsClass->FindPropertyByName("SpectrogramSettings");
+			if (SpectrogramSettingsProperty != nullptr)
+			{
+				const FRackUnitSettingsHelper SettingsHelper(*SpectrogramSettingsProperty);
+
+				AnalyzerParams.AnalyzerType.BindLambda([=]()
+					{
+						return SettingsHelper.GetRackUnitSettings()->AnalyzerType;
+					});
+				AnalyzerParams.FFTAnalyzerFFTSize.BindLambda([=]()
+					{
+						return SettingsHelper.GetRackUnitSettings()->FFTAnalyzerFFTSize;
+					});
+				AnalyzerParams.CQTAnalyzerFFTSize.BindLambda([=]()
+					{
+						return SettingsHelper.GetRackUnitSettings()->CQTAnalyzerFFTSize;
+					});
+				AnalyzerParams.FrequencyAxisPixelBucketMode.BindLambda([=]()
+					{
+						return SettingsHelper.GetRackUnitSettings()->PixelPlotMode;
+					});
+				AnalyzerParams.FrequencyAxisScale.BindLambda([=]()
+					{
+						return SettingsHelper.GetRackUnitSettings()->FrequencyScale;
+					});
+				AnalyzerParams.ColorMap.BindLambda([=]()
+					{
+						return SettingsHelper.GetRackUnitSettings()->ColorMap;
+					});
+				AnalyzerParams.Orientation.BindLambda([=]()
+					{
+						return SettingsHelper.GetRackUnitSettings()->Orientation;
+					});
+
+				AnalyzerParams.OnAnalyzerTypeMenuEntryClicked.BindLambda([=](EAudioSpectrumAnalyzerType SelectedValue)
+					{
+						SettingsHelper.GetRackUnitSettings()->AnalyzerType = SelectedValue;
+						SettingsHelper.SaveConfig();
+					});
+				AnalyzerParams.OnFFTAnalyzerFFTSizeMenuEntryClicked.BindLambda([=](EFFTSize SelectedValue)
+					{
+						SettingsHelper.GetRackUnitSettings()->FFTAnalyzerFFTSize = SelectedValue;
+						SettingsHelper.SaveConfig();
+					});
+				AnalyzerParams.OnCQTAnalyzerFFTSizeMenuEntryClicked.BindLambda([=](EConstantQFFTSizeEnum SelectedValue)
+					{
+						SettingsHelper.GetRackUnitSettings()->CQTAnalyzerFFTSize = SelectedValue;
+						SettingsHelper.SaveConfig();
+					});
+				AnalyzerParams.OnFrequencyAxisPixelBucketModeMenuEntryClicked.BindLambda([=](EAudioSpectrogramFrequencyAxisPixelBucketMode SelectedValue)
+					{
+						SettingsHelper.GetRackUnitSettings()->PixelPlotMode = SelectedValue;
+						SettingsHelper.SaveConfig();
+					});
+				AnalyzerParams.OnFrequencyAxisScaleMenuEntryClicked.BindLambda([=](EAudioSpectrogramFrequencyAxisScale SelectedValue)
+					{
+						SettingsHelper.GetRackUnitSettings()->FrequencyScale = SelectedValue;
+						SettingsHelper.SaveConfig();
+					});
+				AnalyzerParams.OnColorMapMenuEntryClicked.BindLambda([=](EAudioColorGradient SelectedValue)
+					{
+						SettingsHelper.GetRackUnitSettings()->ColorMap = SelectedValue;
+						SettingsHelper.SaveConfig();
+					});
+				AnalyzerParams.OnOrientationMenuEntryClicked.BindLambda([=](EOrientation SelectedValue)
+					{
+						SettingsHelper.GetRackUnitSettings()->Orientation = SelectedValue;
+						SettingsHelper.SaveConfig();
+					});
+			}
+		}
+
+		return MakeShared<FAudioSpectrogram>(AnalyzerParams);
+	}
+} // namespace AudioWidgets
+
+#undef LOCTEXT_NAMESPACE
diff --git a/Engine/Plugins/Runtime/AudioWidgets/Source/AudioWidgets/Private/AudioWidgets.cpp b/Engine/Plugins/Runtime/AudioWidgets/Source/AudioWidgets/Private/AudioWidgets.cpp
--- a/Engine/Plugins/Runtime/AudioWidgets/Source/AudioWidgets/Private/AudioWidgets.cpp
+++ b/Engine/Plugins/Runtime/AudioWidgets/Source/AudioWidgets/Private/AudioWidgets.cpp
@@ -5,6 +5,7 @@
 #include "AudioAnalyzerRackUnitRegistry.h"
 #include "AudioMeter.h"
 #include "AudioOscilloscope.h"
+#include "AudioSpectrogram.h"
 #include "AudioSpectrumAnalyzer.h"
 #include "AudioVectorscope.h"
 #include "AudioWidgetsStyle.h"
@@ -26,6 +27,7 @@ void FAudioWidgetsModule::StartupModule()
 	RegisterAudioAnalyzerRackUnitType(&FAudioMeter::RackUnitTypeInfo);
 	RegisterAudioAnalyzerRackUnitType(&FAudioOscilloscope::RackUnitTypeInfo);
 	RegisterAudioAnalyzerRackUnitType(&FAudioVectorscope::RackUnitTypeInfo);
+	RegisterAudioAnalyzerRackUnitType(&FAudioSpectrogram::RackUnitTypeInfo);
 	RegisterAudioAnalyzerRackUnitType(&FAudioSpectrumAnalyzer::RackUnitTypeInfo);
 }
 
diff --git a/Engine/Plugins/Runtime/AudioWidgets/Source/AudioWidgets/Private/SAudioSpectrogram.cpp b/Engine/Plugins/Runtime/AudioWidgets/Source/AudioWidgets/Private/SAudioSpectrogram.cpp
--- a/Engine/Plugins/Runtime/AudioWidgets/Source/AudioWidgets/Private/SAudioSpectrogram.cpp
+++ b/Engine/Plugins/Runtime/AudioWidgets/Source/AudioWidgets/Private/SAudioSpectrogram.cpp
@@ -24,6 +24,10 @@ void SAudioSpectrogram::Construct(const FArguments& InArgs)
 	Orientation = InArgs._Orientation;
 	bAllowContextMenu = InArgs._AllowContextMenu;
 	OnContextMenuOpening = InArgs._OnContextMenuOpening;
+	OnFrequencyAxisPixelBucketModeMenuEntryClicked = InArgs._OnFrequencyAxisPixelBucketModeMenuEntryClicked;
+	OnFrequencyAxisScaleMenuEntryClicked = InArgs._OnFrequencyAxisScaleMenuEntryClicked;
+	OnColorMapMenuEntryClicked = InArgs._OnColorMapMenuEntryClicked;
+	OnOrientationMenuEntryClicked = InArgs._OnOrientationMenuEntryClicked;
 
 	SpectrogramViewport = MakeShareable(new FAudioSpectrogramViewport());
 }
@@ -182,7 +186,7 @@ TSharedRef<SWidget> SAudioSpectrogram::BuildDefaultContextMenu()
 
 	MenuBuilder.BeginSection(ContextMenuExtensionHook, LOCTEXT("DisplayOptions", "Display Options"));
 
-	if (!FrequencyAxisPixelBucketMode.IsBound())
+	if (OnFrequencyAxisPixelBucketModeMenuEntryClicked.IsBound() || !FrequencyAxisPixelBucketMode.IsBound())
 	{
 		MenuBuilder.AddSubMenu(
 			LOCTEXT("FrequencyAxisPixelBucketMode", "Pixel Plot Mode"),
@@ -190,7 +194,7 @@ TSharedRef<SWidget> SAudioSpectrogram::BuildDefaultContextMenu()
 			FNewMenuDelegate::CreateSP(this, &SAudioSpectrogram::BuildFrequencyAxisPixelBucketModeSubMenu));
 	}
 
-	if (!FrequencyAxisScale.IsBound())
+	if (OnFrequencyAxisScaleMenuEntryClicked.IsBound() || !FrequencyAxisScale.IsBound())
 	{
 		MenuBuilder.AddSubMenu(
 			LOCTEXT("FrequencyAxisScale", "Frequency Scale"),
@@ -198,7 +202,7 @@ TSharedRef<SWidget> SAudioSpectrogram::BuildDefaultContextMenu()
 			FNewMenuDelegate::CreateSP(this, &SAudioSpectrogram::BuildFrequencyAxisScaleSubMenu));
 	}
 
-	if (!ColorMap.IsBound())
+	if (OnColorMapMenuEntryClicked.IsBound() || !ColorMap.IsBound())
 	{
 		MenuBuilder.AddSubMenu(
 			LOCTEXT("ColorMap", "Color Map"),
@@ -206,7 +210,7 @@ TSharedRef<SWidget> SAudioSpectrogram::BuildDefaultContextMenu()
 			FNewMenuDelegate::CreateSP(this, &SAudioSpectrogram::BuildColorMapSubMenu));
 	}
 
-	if (!Orientation.IsBound())
+	if (OnOrientationMenuEntryClicked.IsBound() || !Orientation.IsBound())
 	{
 		MenuBuilder.AddSubMenu(
 			LOCTEXT("Orientation", "Orientation"),
@@ -236,7 +240,15 @@ void SAudioSpectrogram::BuildColorMapSubMenu(FMenuBuilder& SubMenu)
 #endif
 			FSlateIcon(),
 			FUIAction(
-				FExecuteAction::CreateSP(this, &SAudioSpectrogram::SetColorMap, EnumValue),
+				FExecuteAction::CreateSPLambda(this, [this, EnumValue]()
+					{
+						if (!ColorMap.IsBound())
+						{
+							ColorMap = EnumValue;
+						}
+
+						OnColorMapMenuEntryClicked.ExecuteIfBound(EnumValue);
+					}),
 				FCanExecuteAction(),
 				FIsActionChecked::CreateSPLambda(this, [this, EnumValue]() { return (ColorMap.Get() == EnumValue); })
 			),
@@ -262,7 +274,15 @@ void SAudioSpectrogram::BuildFrequencyAxisScaleSubMenu(FMenuBuilder& SubMenu)
 #endif
 			FSlateIcon(),
 			FUIAction(
-				FExecuteAction::CreateSP(this, &SAudioSpectrogram::SetFrequencyAxisScale, EnumValue),
+				FExecuteAction::CreateSPLambda(this, [this, EnumValue]()
+					{
+						if (!FrequencyAxisScale.IsBound())
+						{
+							FrequencyAxisScale = EnumValue;
+						}
+
+						OnFrequencyAxisScaleMenuEntryClicked.ExecuteIfBound(EnumValue);
+					}),
 				FCanExecuteAction(),
 				FIsActionChecked::CreateSPLambda(this, [this, EnumValue]() { return (FrequencyAxisScale.Get() == EnumValue); })
 			),
@@ -288,7 +308,15 @@ void SAudioSpectrogram::BuildFrequencyAxisPixelBucketModeSubMenu(FMenuBuilder& S
 #endif
 			FSlateIcon(),
 			FUIAction(
-				FExecuteAction::CreateSP(this, &SAudioSpectrogram::SetFrequencyAxisPixelBucketMode, EnumValue),
+				FExecuteAction::CreateSPLambda(this, [this, EnumValue]()
+					{
+						if (!FrequencyAxisPixelBucketMode.IsBound())
+						{
+							FrequencyAxisPixelBucketMode = EnumValue;
+						}
+
+						OnFrequencyAxisPixelBucketModeMenuEntryClicked.ExecuteIfBound(EnumValue);
+					}),
 				FCanExecuteAction(),
 				FIsActionChecked::CreateSPLambda(this, [this, EnumValue]() { return (FrequencyAxisPixelBucketMode.Get() == EnumValue); })
 			),
@@ -314,7 +342,15 @@ void SAudioSpectrogram::BuildOrientationSubMenu(FMenuBuilder& SubMenu)
 #endif
 			FSlateIcon(),
 			FUIAction(
-				FExecuteAction::CreateSP(this, &SAudioSpectrogram::SetOrientation, EnumValue),
+				FExecuteAction::CreateSPLambda(this, [this, EnumValue]()
+					{
+						if (!Orientation.IsBound())
+						{
+							Orientation = EnumValue;
+						}
+
+						OnOrientationMenuEntryClicked.ExecuteIfBound(EnumValue);
+					}),
 				FCanExecuteAction(),
 				FIsActionChecked::CreateSPLambda(this, [this, EnumValue]() { return (Orientation.Get() == EnumValue); })
 			),
diff --git a/Engine/Plugins/Runtime/AudioWidgets/Source/AudioWidgets/Public/AudioSpectrogram.h b/Engine/Plugins/Runtime/AudioWidgets/Source/AudioWidgets/Public/AudioSpectrogram.h
new file mode 100644
--- /dev/null
+++ b/Engine/Plugins/Runtime/AudioWidgets/Source/AudioWidgets/Public/AudioSpectrogram.h
@@ -0,0 +1,148 @@
+// Copyright Epic Games, Inc. All Rights Reserved.
+#pragma once
+
+#include "AudioAnalyzerRack.h"
+#include "AudioSpectrumAnalyzer.h"
+#include "SAudioSpectrogram.h"
+
+#include "AudioSpectrogram.generated.h"
+
+namespace AudioWidgets
+{
+	/**
+	 * Constructor parameters for the analyzer.
+	 */
+	struct AUDIOWIDGETS_API FAudioSpectrogramParams
+	{
+		int32 NumChannels = 1;
+		Audio::FDeviceId AudioDeviceId = INDEX_NONE;
+		TObjectPtr<UAudioBus> ExternalAudioBus = nullptr;
+
+		TAttribute<EAudioSpectrumAnalyzerType> AnalyzerType = EAudioSpectrumAnalyzerType::FFT;
+		TAttribute<EFFTSize> FFTAnalyzerFFTSize = EFFTSize::Max;
+		TAttribute<EConstantQFFTSizeEnum> CQTAnalyzerFFTSize = EConstantQFFTSizeEnum::XXLarge;
+		TAttribute<EAudioSpectrogramFrequencyAxisPixelBucketMode> FrequencyAxisPixelBucketMode = EAudioSpectrogramFrequencyAxisPixelBucketMode::Average;
+		TAttribute<EAudioSpectrogramFrequencyAxisScale> FrequencyAxisScale = EAudioSpectrogramFrequencyAxisScale::Logarithmic;
+		TAttribute<EAudioColorGradient> ColorMap = EAudioColorGradient::BlackToWhite;
+		TAttribute<EOrientation> Orientation = EOrientation::Orient_Horizontal;
+
+		FOnAnalyzerTypeMenuEntryClicked OnAnalyzerTypeMenuEntryClicked;
+		FOnFFTAnalyzerFFTSizeMenuEntryClicked OnFFTAnalyzerFFTSizeMenuEntryClicked;
+		FOnCQTAnalyzerFFTSizeMenuEntryClicked OnCQTAnalyzerFFTSizeMenuEntryClicked;
+		FOnSpectrogramFrequencyAxisPixelBucketModeMenuEntryClicked OnFrequencyAxisPixelBucketModeMenuEntryClicked;
+		FOnSpectrogramFrequencyAxisScaleMenuEntryClicked OnFrequencyAxisScaleMenuEntryClicked;
+		FOnSpectrogramColorMapMenuEntryClicked OnColorMapMenuEntryClicked;
+		FOnSpectrogramOrientationMenuEntryClicked OnOrientationMenuEntryClicked;
+	};
+
+	/**
+	 * Owns an analyzer and a corresponding Slate widget for displaying the resulting spectra.
+	 * Can either create an Audio Bus to analyze, or analyze the given Bus.
+	 */
+	class AUDIOWIDGETS_API FAudioSpectrogram : public IAudioAnalyzerRackUnit
+	{
+	public:
+		static const FAudioAnalyzerRackUnitTypeInfo RackUnitTypeInfo;
+		
+		FAudioSpectrogram(const FAudioSpectrogramParams& Params);
+		~FAudioSpectrogram();
+
+		UAudioBus* GetAudioBus() const;
+
+		TSharedRef<SWidget> GetWidget() const;
+
+		void Init(int32 InNumChannels, Audio::FDeviceId InAudioDeviceId, TObjectPtr<UAudioBus> InExternalAudioBus = nullptr);
+
+		// Begin IAudioAnalyzerRackUnit overrides.
+		virtual void SetAudioBusInfo(const FAudioBusInfo& AudioBusInfo) override
+		{
+			Init(AudioBusInfo.AudioBus->GetNumChannels(), AudioBusInfo.AudioDeviceId, AudioBusInfo.AudioBus);
+		}
+
+		virtual TSharedRef<SDockTab> SpawnTab(const FSpawnTabArgs& Args) const override;
+		// End IAudioAnalyzerRackUnit overrides.
+
+	protected:
+		void StartAnalyzing(const EAudioSpectrumAnalyzerType InAnalyzerType);
+		void StopAnalyzing();
+
+		void OnSpectrumResults(USynesthesiaSpectrumAnalyzer* InSpectrumAnalyzer, int32 ChannelIndex, const TArray<FSynesthesiaSpectrumResults>& InSpectrumResultsArray);
+		void OnConstantQResults(UConstantQAnalyzer* InSpectrumAnalyzer, int32 ChannelIndex, const TArray<FConstantQResults>& InSpectrumResultsArray);
+
+		void ExtendSpectrumPlotContextMenu(FMenuBuilder& MenuBuilder);
+		void BuildAnalyzerTypeSubMenu(FMenuBuilder& SubMenu);
+		void BuildFFTSizeSubMenu(FMenuBuilder& SubMenu);
+
+		EActiveTimerReturnType Update(double InCurrentTime, float InDeltaTime);
+
+	private:
+		static TSharedRef<IAudioAnalyzerRackUnit> MakeRackUnit(const FAudioAnalyzerRackUnitConstructParams& Params);
+
+		void CreateSynesthesiaSpectrumAnalyzer();
+		void ReleaseSynesthesiaSpectrumAnalyzer();
+
+		void CreateConstantQAnalyzer();
+		void ReleaseConstantQAnalyzer();
+
+		void Teardown();
+
+		/** Audio analyzer objects. */
+		TStrongObjectPtr<USynesthesiaSpectrumAnalyzer> SpectrumAnalyzer;
+		TStrongObjectPtr<UConstantQAnalyzer> ConstantQAnalyzer;
+
+		/** The audio bus used for analysis. */
+		TStrongObjectPtr<UAudioBus> AudioBus;
+
+		/** Handles for results delegate for analyzers. */
+		FDelegateHandle SpectrumResultsDelegateHandle;
+		FDelegateHandle ConstantQResultsDelegateHandle;
+
+		/** Analyzer settings. */
+		TStrongObjectPtr<USynesthesiaSpectrumAnalysisSettings> SpectrumAnalysisSettings;
+		TStrongObjectPtr<UConstantQSettings> ConstantQSettings;
+
+		/** Slate widget for spectrum display */
+		TSharedPtr<SAudioSpectrogram> Widget;
+		TSharedPtr<const FExtensionBase> ContextMenuExtension;
+		TSharedPtr<FActiveTimerHandle> ActiveTimer;
+
+		Audio::FDeviceId AudioDeviceId = INDEX_NONE;
+		bool bUseExternalAudioBus = false;
+
+		TOptional<EAudioSpectrumAnalyzerType> ActiveAnalyzerType;
+		TAttribute<EAudioSpectrumAnalyzerType> AnalyzerType;
+		TAttribute<EFFTSize> FFTAnalyzerFFTSize;
+		TAttribute<EConstantQFFTSizeEnum> CQTAnalyzerFFTSize;
+
+		FOnAnalyzerTypeMenuEntryClicked OnAnalyzerTypeMenuEntryClicked;
+		FOnFFTAnalyzerFFTSizeMenuEntryClicked OnFFTAnalyzerFFTSizeMenuEntryClicked;
+		FOnCQTAnalyzerFFTSizeMenuEntryClicked OnCQTAnalyzerFFTSizeMenuEntryClicked;
+	};
+} // namespace AudioWidgets
+
+USTRUCT()
+struct FSpectrogramRackUnitSettings
+{
+	GENERATED_BODY()
+
+	UPROPERTY(EditAnywhere, config, Category = Spectrogram)
+	EAudioSpectrumAnalyzerType AnalyzerType = EAudioSpectrumAnalyzerType::FFT;
+
+	UPROPERTY(EditAnywhere, config, Category = Spectrogram, meta = (DisplayName = "FFT Size (FFT Analyzer)"))
+	EFFTSize FFTAnalyzerFFTSize = EFFTSize::Max;
+
+	UPROPERTY(EditAnywhere, config, Category = Spectrogram, meta = (DisplayName = "FFT Size (CQT Analyzer)"))
+	EConstantQFFTSizeEnum CQTAnalyzerFFTSize = EConstantQFFTSizeEnum::XXLarge;
+
+	UPROPERTY(EditAnywhere, config, Category = Spectrogram)
+	EAudioSpectrogramFrequencyAxisPixelBucketMode PixelPlotMode = EAudioSpectrogramFrequencyAxisPixelBucketMode::Average;
+
+	UPROPERTY(EditAnywhere, config, Category = Spectrogram)
+	EAudioSpectrogramFrequencyAxisScale FrequencyScale = EAudioSpectrogramFrequencyAxisScale::Logarithmic;
+
+	UPROPERTY(EditAnywhere, config, Category = Spectrogram)
+	EAudioColorGradient ColorMap = EAudioColorGradient::BlackToWhite;
+
+	UPROPERTY(EditAnywhere, config, Category = Spectrogram)
+	TEnumAsByte<EOrientation> Orientation = EOrientation::Orient_Horizontal;
+};
diff --git a/Engine/Plugins/Runtime/AudioWidgets/Source/AudioWidgets/Public/SAudioSpectrogram.h b/Engine/Plugins/Runtime/AudioWidgets/Source/AudioWidgets/Public/SAudioSpectrogram.h
--- a/Engine/Plugins/Runtime/AudioWidgets/Source/AudioWidgets/Public/SAudioSpectrogram.h
+++ b/Engine/Plugins/Runtime/AudioWidgets/Source/AudioWidgets/Public/SAudioSpectrogram.h
@@ -11,6 +11,11 @@
 struct FSynesthesiaSpectrumResults;
 struct FConstantQResults;
 
+DECLARE_DELEGATE_OneParam(FOnSpectrogramFrequencyAxisPixelBucketModeMenuEntryClicked, EAudioSpectrogramFrequencyAxisPixelBucketMode);
+DECLARE_DELEGATE_OneParam(FOnSpectrogramFrequencyAxisScaleMenuEntryClicked, EAudioSpectrogramFrequencyAxisScale);
+DECLARE_DELEGATE_OneParam(FOnSpectrogramColorMapMenuEntryClicked, EAudioColorGradient);
+DECLARE_DELEGATE_OneParam(FOnSpectrogramOrientationMenuEntryClicked, EOrientation);
+
 /**
  * Slate Widget for rendering a time-frequency representation of a series of audio power spectra.
  */
@@ -38,6 +43,10 @@ public:
 		SLATE_ATTRIBUTE(EOrientation, Orientation)
 		SLATE_ATTRIBUTE(bool, AllowContextMenu)
 		SLATE_EVENT(FOnContextMenuOpening, OnContextMenuOpening)
+		SLATE_EVENT(FOnSpectrogramFrequencyAxisPixelBucketModeMenuEntryClicked, OnFrequencyAxisPixelBucketModeMenuEntryClicked)
+		SLATE_EVENT(FOnSpectrogramFrequencyAxisScaleMenuEntryClicked, OnFrequencyAxisScaleMenuEntryClicked)
+		SLATE_EVENT(FOnSpectrogramColorMapMenuEntryClicked, OnColorMapMenuEntryClicked)
+		SLATE_EVENT(FOnSpectrogramOrientationMenuEntryClicked, OnOrientationMenuEntryClicked)
 	SLATE_END_ARGS()
 
 	/** Constructs this widget with InArgs */
@@ -94,6 +103,10 @@ private:
 	TAttribute<EOrientation> Orientation;
 	TAttribute<bool> bAllowContextMenu;
 	FOnContextMenuOpening OnContextMenuOpening;
+	FOnSpectrogramFrequencyAxisPixelBucketModeMenuEntryClicked OnFrequencyAxisPixelBucketModeMenuEntryClicked;
+	FOnSpectrogramFrequencyAxisScaleMenuEntryClicked OnFrequencyAxisScaleMenuEntryClicked;
+	FOnSpectrogramColorMapMenuEntryClicked OnColorMapMenuEntryClicked;
+	FOnSpectrogramOrientationMenuEntryClicked OnOrientationMenuEntryClicked;
 
 	TSharedPtr<FAudioSpectrogramViewport> SpectrogramViewport;
 };
diff --git a/Engine/Plugins/Runtime/Metasound/Source/MetasoundEditor/Private/MetasoundEditorSettings.h b/Engine/Plugins/Runtime/Metasound/Source/MetasoundEditor/Private/MetasoundEditorSettings.h
--- a/Engine/Plugins/Runtime/Metasound/Source/MetasoundEditor/Private/MetasoundEditorSettings.h
+++ b/Engine/Plugins/Runtime/Metasound/Source/MetasoundEditor/Private/MetasoundEditorSettings.h
@@ -1,6 +1,7 @@
 // Copyright Epic Games, Inc. All Rights Reserved.
 #pragma once
 
+#include "AudioSpectrogram.h"
 #include "AudioSpectrumAnalyzer.h"
 #include "MetasoundFrontendDocument.h"
 #include "Misc/CoreDefines.h"
@@ -221,6 +222,10 @@ public:
 	UPROPERTY(EditAnywhere, config, Category = NodeTitleColors)
 	FLinearColor VariableNodeTitleColor;
 
+	/** Settings for metasound output spectrogram widget */
+	UPROPERTY(EditAnywhere, config, Category = Spectrogram, meta = (ShowOnlyInnerProperties))
+	FSpectrogramRackUnitSettings SpectrogramSettings;
+
 	/** Settings for metasound output spectrum analyzer widget */
 	UPROPERTY(EditAnywhere, config, Category = SpectrumAnalyzer, meta = (ShowOnlyInnerProperties))
 	FSpectrumAnalyzerRackUnitSettings SpectrumAnalyzerSettings;
