Fix for EngineTest break

diff --git a/Engine/Plugins/Runtime/AudioWidgets/Source/AudioWidgets/Private/SparseSampledSequenceTransportCoordinator.cpp b/Engine/Plugins/Runtime/AudioWidgets/Source/AudioWidgets/Private/SparseSampledSequenceTransportCoordinator.cpp
--- a/Engine/Plugins/Runtime/AudioWidgets/Source/AudioWidgets/Private/SparseSampledSequenceTransportCoordinator.cpp
+++ b/Engine/Plugins/Runtime/AudioWidgets/Source/AudioWidgets/Private/SparseSampledSequenceTransportCoordinator.cpp
@@ -96,6 +96,11 @@ void FSparseSampledSequenceTransportCoordinator::UpdateZoomRatioAndDisplayRange(
 	MoveFocusPoint(NewPlayheadPosition);
 }
 
+void FSparseSampledSequenceTransportCoordinator::Stop()
+{
+	SetProgressRatio(0.f);
+}
+
 const float FSparseSampledSequenceTransportCoordinator::GetFocusPoint() const
 {
 	return FocusPoint;
diff --git a/Engine/Plugins/Runtime/AudioWidgets/Source/AudioWidgets/Public/SparseSampledSequenceTransportCoordinator.h b/Engine/Plugins/Runtime/AudioWidgets/Source/AudioWidgets/Public/SparseSampledSequenceTransportCoordinator.h
--- a/Engine/Plugins/Runtime/AudioWidgets/Source/AudioWidgets/Public/SparseSampledSequenceTransportCoordinator.h
+++ b/Engine/Plugins/Runtime/AudioWidgets/Source/AudioWidgets/Public/SparseSampledSequenceTransportCoordinator.h
@@ -34,6 +34,7 @@ public:
 	float ConvertAbsoluteRatioToZoomed(const float InAbsoluteRatio) const override;
 	float ConvertZoomedRatioToAbsolute(const float InZoomedRatio) const override;
 	void UpdatePlaybackRange(const TRange<double>& NewRange);
+	void Stop();
 
 private:
 	FORCEINLINE void MoveFocusPoint(const double InFocusPoint);
