[AudioWidgets] Spectrogram analyzer rack unit - background color is now cleared to the silence value instead of being transparent. Also, fix for null texture if spectrogram is resized to zero rows.

diff --git a/Engine/Plugins/Runtime/AudioWidgets/Source/AudioWidgets/Private/AudioSpectrogram.cpp b/Engine/Plugins/Runtime/AudioWidgets/Source/AudioWidgets/Private/AudioSpectrogram.cpp
--- a/Engine/Plugins/Runtime/AudioWidgets/Source/AudioWidgets/Private/AudioSpectrogram.cpp
+++ b/Engine/Plugins/Runtime/AudioWidgets/Source/AudioWidgets/Private/AudioSpectrogram.cpp
@@ -63,6 +63,7 @@ namespace AudioWidgets
 			.FrequencyAxisPixelBucketMode(Params.FrequencyAxisPixelBucketMode)
 			.ColorMap(Params.ColorMap)
 			.Orientation(Params.Orientation)
+			.FillBackground(true)
 			.OnFrequencyAxisPixelBucketModeMenuEntryClicked(Params.OnFrequencyAxisPixelBucketModeMenuEntryClicked)
 			.OnFrequencyAxisScaleMenuEntryClicked(Params.OnFrequencyAxisScaleMenuEntryClicked)
 			.OnColorMapMenuEntryClicked(Params.OnColorMapMenuEntryClicked)
diff --git a/Engine/Plugins/Runtime/AudioWidgets/Source/AudioWidgets/Private/SAudioSpectrogram.cpp b/Engine/Plugins/Runtime/AudioWidgets/Source/AudioWidgets/Private/SAudioSpectrogram.cpp
--- a/Engine/Plugins/Runtime/AudioWidgets/Source/AudioWidgets/Private/SAudioSpectrogram.cpp
+++ b/Engine/Plugins/Runtime/AudioWidgets/Source/AudioWidgets/Private/SAudioSpectrogram.cpp
@@ -2,6 +2,7 @@
 
 #include "SAudioSpectrogram.h"
 
+#include "AudioColorMapper.h"
 #include "ConstantQ.h"
 #include "Framework/Application/SlateApplication.h"
 #include "Framework/MultiBox/MultiBoxBuilder.h"
@@ -23,6 +24,7 @@ void SAudioSpectrogram::Construct(const FArguments& InArgs)
 	FrequencyAxisPixelBucketMode = InArgs._FrequencyAxisPixelBucketMode;
 	Orientation = InArgs._Orientation;
 	bAllowContextMenu = InArgs._AllowContextMenu;
+	bFillBackground = InArgs._FillBackground;
 	OnContextMenuOpening = InArgs._OnContextMenuOpening;
 	OnFrequencyAxisPixelBucketModeMenuEntryClicked = InArgs._OnFrequencyAxisPixelBucketModeMenuEntryClicked;
 	OnFrequencyAxisScaleMenuEntryClicked = InArgs._OnFrequencyAxisScaleMenuEntryClicked;
@@ -161,7 +163,7 @@ int32 SAudioSpectrogram::OnPaint(const FPaintArgs& Args, const FGeometry& Allott
 	// Create spectrogram render params, setting viewport history size to exact pixel size:
 	const FAudioSpectrogramViewportRenderParams RenderParams
 	{
-		.NumRows = (bRotateViewport) ? RenderWidth : RenderHeight,
+		.NumRows = FMath::Max((bRotateViewport) ? RenderWidth : RenderHeight, 1),
 		.NumPixelsPerRow = FMath::Max((bRotateViewport) ? RenderHeight : RenderWidth, 2),
 		.ViewMinFrequency = ViewMinFrequency.Get(),
 		.ViewMaxFrequency = ViewMaxFrequency.Get(),
@@ -173,6 +175,15 @@ int32 SAudioSpectrogram::OnPaint(const FPaintArgs& Args, const FGeometry& Allott
 	};
 	SpectrogramViewport->SetRenderParams(RenderParams);
 
+	if (bFillBackground.Get())
+	{
+		// Fill allotted geometry with colormap minimum value:
+		const FSlateBrush BackgroundBrush;
+		const FLinearColor BoxColor = FAudioColorMapper(0.0f, 1.0f, RenderParams.ColorMap).GetColorFromValue(0.0f);
+		FSlateDrawElement::MakeBox(OutDrawElements, LayerId, AllottedGeometry.ToPaintGeometry(), &BackgroundBrush, ESlateDrawEffect::None, BoxColor);
+		LayerId++;
+	}
+
 	// Create the viewport using child geometry:
 	FSlateDrawElement::MakeViewport(OutDrawElements, LayerId, ChildGeometry.ToPaintGeometry(), SpectrogramViewport, ESlateDrawEffect::NoGamma);
 
diff --git a/Engine/Plugins/Runtime/AudioWidgets/Source/AudioWidgets/Public/SAudioSpectrogram.h b/Engine/Plugins/Runtime/AudioWidgets/Source/AudioWidgets/Public/SAudioSpectrogram.h
--- a/Engine/Plugins/Runtime/AudioWidgets/Source/AudioWidgets/Public/SAudioSpectrogram.h
+++ b/Engine/Plugins/Runtime/AudioWidgets/Source/AudioWidgets/Public/SAudioSpectrogram.h
@@ -32,6 +32,7 @@ public:
 		, _FrequencyAxisPixelBucketMode(EAudioSpectrogramFrequencyAxisPixelBucketMode::Average)
 		, _Orientation(EOrientation::Orient_Horizontal)
 		, _AllowContextMenu(true)
+		, _FillBackground(false)
 	{}
 		SLATE_ATTRIBUTE(float, ViewMinFrequency)
 		SLATE_ATTRIBUTE(float, ViewMaxFrequency)
@@ -42,6 +43,7 @@ public:
 		SLATE_ATTRIBUTE(EAudioSpectrogramFrequencyAxisPixelBucketMode, FrequencyAxisPixelBucketMode)
 		SLATE_ATTRIBUTE(EOrientation, Orientation)
 		SLATE_ATTRIBUTE(bool, AllowContextMenu)
+		SLATE_ATTRIBUTE(bool, FillBackground)
 		SLATE_EVENT(FOnContextMenuOpening, OnContextMenuOpening)
 		SLATE_EVENT(FOnSpectrogramFrequencyAxisPixelBucketModeMenuEntryClicked, OnFrequencyAxisPixelBucketModeMenuEntryClicked)
 		SLATE_EVENT(FOnSpectrogramFrequencyAxisScaleMenuEntryClicked, OnFrequencyAxisScaleMenuEntryClicked)
@@ -102,6 +104,7 @@ private:
 	TAttribute<EAudioSpectrogramFrequencyAxisPixelBucketMode> FrequencyAxisPixelBucketMode;
 	TAttribute<EOrientation> Orientation;
 	TAttribute<bool> bAllowContextMenu;
+	TAttribute<bool> bFillBackground;
 	FOnContextMenuOpening OnContextMenuOpening;
 	FOnSpectrogramFrequencyAxisPixelBucketModeMenuEntryClicked OnFrequencyAxisPixelBucketModeMenuEntryClicked;
 	FOnSpectrogramFrequencyAxisScaleMenuEntryClicked OnFrequencyAxisScaleMenuEntryClicked;
