[AudioWidgets] Fix for editor crash with undocked analyzers.

diff --git a/Engine/Plugins/Runtime/AudioWidgets/Source/AudioWidgets/Private/AudioAnalyzerRack.cpp b/Engine/Plugins/Runtime/AudioWidgets/Source/AudioWidgets/Private/AudioAnalyzerRack.cpp
--- a/Engine/Plugins/Runtime/AudioWidgets/Source/AudioWidgets/Private/AudioAnalyzerRack.cpp
+++ b/Engine/Plugins/Runtime/AudioWidgets/Source/AudioWidgets/Private/AudioAnalyzerRack.cpp
@@ -46,6 +46,14 @@ namespace AudioWidgets
 		RackUnitConstructParams.EditorSettingsClass = Params.EditorSettingsClass;
 	}
 
+	FAudioAnalyzerRack::~FAudioAnalyzerRack()
+	{
+		if (TabManager.IsValid())
+		{
+			TabManager->CloseAllAreas();
+		}
+	}
+
 	void FAudioAnalyzerRack::Init(int32 InNumChannels, Audio::FDeviceId InAudioDeviceId)
 	{
 		const EAudioBusChannels AudioBusChannels = AudioBusUtils::ConvertIntToEAudioBusChannels(InNumChannels);
@@ -73,6 +81,11 @@ namespace AudioWidgets
 
 	TSharedRef<SWidget> FAudioAnalyzerRack::CreateWidget(TSharedRef<SDockTab> DockTab)
 	{
+		if (TabManager.IsValid())
+		{
+			TabManager->CloseAllAreas();
+		}
+
 		// Create a TabManager:
 		TabManager = FGlobalTabmanager::Get()->NewTabManager(DockTab);
 		TabManager->SetOnPersistLayout(FTabManager::FOnPersistLayout::CreateStatic(&FAudioAnalyzerRack::SaveTabLayout));
@@ -151,6 +164,15 @@ namespace AudioWidgets
 					]
 			];
 
+		// If the dock tab that contains the analyzer rack is closed, also close any undocked analyzer rack units:
+		DockTab->SetOnTabClosed(SDockTab::FOnTabClosedCallback::CreateLambda([AnalyzerRackTabManager = TabManager.ToWeakPtr()](TSharedRef<SDockTab> AnalyzerRackDockTab)
+			{
+				if (TSharedPtr<FTabManager> PinnedTabManager = AnalyzerRackTabManager.Pin())
+				{
+					PinnedTabManager->CloseAllAreas();
+				}
+			}));
+
 		return Widget;
 	}
 
diff --git a/Engine/Plugins/Runtime/AudioWidgets/Source/AudioWidgets/Public/AudioAnalyzerRack.h b/Engine/Plugins/Runtime/AudioWidgets/Source/AudioWidgets/Public/AudioAnalyzerRack.h
--- a/Engine/Plugins/Runtime/AudioWidgets/Source/AudioWidgets/Public/AudioAnalyzerRack.h
+++ b/Engine/Plugins/Runtime/AudioWidgets/Source/AudioWidgets/Public/AudioAnalyzerRack.h
@@ -82,7 +82,7 @@ namespace AudioWidgets
 		};
 
 		FAudioAnalyzerRack(const FRackConstructParams& Params);
-		virtual ~FAudioAnalyzerRack() = default;
+		virtual ~FAudioAnalyzerRack();
 
 		void Init(int32 InNumChannels, Audio::FDeviceId InAudioDeviceId);
 		void DestroyAnalyzers();
