[AudioWidgets] Improve safety in analyzer rack code.

diff --git a/Engine/Plugins/Runtime/AudioWidgets/Source/AudioWidgets/Private/AudioAnalyzerRack.cpp b/Engine/Plugins/Runtime/AudioWidgets/Source/AudioWidgets/Private/AudioAnalyzerRack.cpp
--- a/Engine/Plugins/Runtime/AudioWidgets/Source/AudioWidgets/Private/AudioAnalyzerRack.cpp
+++ b/Engine/Plugins/Runtime/AudioWidgets/Source/AudioWidgets/Private/AudioAnalyzerRack.cpp
@@ -50,6 +50,8 @@ namespace AudioWidgets
 	{
 		if (TabManager.IsValid())
 		{
+			TabManager->UnregisterAllTabSpawners();
+			TabManager->SetOnPersistLayout(nullptr);
 			TabManager->CloseAllAreas();
 		}
 	}
@@ -83,6 +85,8 @@ namespace AudioWidgets
 	{
 		if (TabManager.IsValid())
 		{
+			TabManager->UnregisterAllTabSpawners();
+			TabManager->SetOnPersistLayout(nullptr);
 			TabManager->CloseAllAreas();
 		}
 
@@ -169,6 +173,7 @@ namespace AudioWidgets
 			{
 				if (TSharedPtr<FTabManager> PinnedTabManager = AnalyzerRackTabManager.Pin())
 				{
+					PinnedTabManager->SetOnPersistLayout(nullptr);
 					PinnedTabManager->CloseAllAreas();
 				}
 			}));
