Fix for UAudioMaterialSlider not updating its internal value when moving the slider

#rnx

diff --git a/Engine/Plugins/Runtime/AudioWidgets/Source/AudioWidgets/Private/AudioMaterialSlate/AudioMaterialSlider.cpp b/Engine/Plugins/Runtime/AudioWidgets/Source/AudioWidgets/Private/AudioMaterialSlate/AudioMaterialSlider.cpp
--- a/Engine/Plugins/Runtime/AudioWidgets/Source/AudioWidgets/Private/AudioMaterialSlate/AudioMaterialSlider.cpp
+++ b/Engine/Plugins/Runtime/AudioWidgets/Source/AudioWidgets/Private/AudioMaterialSlate/AudioMaterialSlider.cpp
@@ -59,11 +59,6 @@ void UAudioMaterialSlider::SetValue(float InValue)
 	if (Slider.IsValid())
 	{
 		Slider->SetValue(InValue);
-	}
-
-	if (Value != InValue)
-	{
-		Value = InValue;
 		HandleOnValueChanged(InValue);
 	}
 }
@@ -161,7 +156,11 @@ TSharedRef<SWidget> UAudioMaterialSlider::RebuildWidget()
 
 void UAudioMaterialSlider::HandleOnValueChanged(float InValue)
 {
-	OnValueChanged.Broadcast(InValue);
+	if (Value != InValue)
+	{
+		Value = InValue;
+		OnValueChanged.Broadcast(InValue);	
+	}
 }
 
 #undef LOCTEXT_NAMESPACE
