[AudioWidgets] Fixes for analyzer crashes on opening metasound editor or audio insights when editor was run with -nosound.

diff --git a/Engine/Plugins/Runtime/AudioWidgets/Source/AudioWidgets/Private/AudioAnalyzerRack.cpp b/Engine/Plugins/Runtime/AudioWidgets/Source/AudioWidgets/Private/AudioAnalyzerRack.cpp
--- a/Engine/Plugins/Runtime/AudioWidgets/Source/AudioWidgets/Private/AudioAnalyzerRack.cpp
+++ b/Engine/Plugins/Runtime/AudioWidgets/Source/AudioWidgets/Private/AudioAnalyzerRack.cpp
@@ -313,9 +313,6 @@ namespace AudioWidgets
 
 	TSharedRef<IAudioAnalyzerRackUnit> FAudioAnalyzerRack::MakeRackUnit(FName RackUnitTypeName)
 	{
-		// FAudioAnalyzerRack::Init(...) is required to have been called first:
-		ensure(AudioBus.IsValid());
-
 		TSharedRef<IAudioAnalyzerRackUnit> RackUnit = FAudioAnalyzerRackUnitRegistry::Get().MakeRackUnit(RackUnitTypeName, RackUnitConstructParams);
 
 		if (bIsProcessingStarted)
diff --git a/Engine/Plugins/Runtime/AudioWidgets/Source/AudioWidgets/Private/AudioMeter.cpp b/Engine/Plugins/Runtime/AudioWidgets/Source/AudioWidgets/Private/AudioMeter.cpp
--- a/Engine/Plugins/Runtime/AudioWidgets/Source/AudioWidgets/Private/AudioMeter.cpp
+++ b/Engine/Plugins/Runtime/AudioWidgets/Source/AudioWidgets/Private/AudioMeter.cpp
@@ -236,24 +236,28 @@ namespace AudioWidgets
 
 	void FAudioMeter::Init(const int32 InNumChannels, const Audio::FDeviceId InAudioDeviceId, const TObjectPtr<UAudioBus> InExternalAudioBus)
 	{
-		check(InNumChannels > 0);
-
 		Teardown();
 
 		Settings = TStrongObjectPtr(NewObject<UMeterSettings>());
 		Settings->PeakHoldTime = 4000.0f;
 
-		Analyzer = TStrongObjectPtr(NewObject<UMeterAnalyzer>());
-		Analyzer->Settings = Settings.Get();
+		// Only create analyzers etc if we have an audio device:
+		if (InAudioDeviceId != FAudioBusInfo::InvalidAudioDeviceId)
+		{
+			check(InNumChannels > 0);
+
+			Analyzer = TStrongObjectPtr(NewObject<UMeterAnalyzer>());
+			Analyzer->Settings = Settings.Get();
 
-		bUseExternalAudioBus = InExternalAudioBus != nullptr;
+			bUseExternalAudioBus = InExternalAudioBus != nullptr;
 
-		AudioBus = bUseExternalAudioBus ? TStrongObjectPtr(InExternalAudioBus.Get()) : TStrongObjectPtr(NewObject<UAudioBus>());
-		AudioBus->AudioBusChannels = EAudioBusChannels(InNumChannels - 1);
+			AudioBus = bUseExternalAudioBus ? TStrongObjectPtr(InExternalAudioBus.Get()) : TStrongObjectPtr(NewObject<UAudioBus>());
+			AudioBus->AudioBusChannels = EAudioBusChannels(InNumChannels - 1);
 
-		ResultsDelegateHandle = Analyzer->OnLatestPerChannelMeterResultsNative.AddRaw(this, &FAudioMeter::OnMeterOutput);
+			ResultsDelegateHandle = Analyzer->OnLatestPerChannelMeterResultsNative.AddRaw(this, &FAudioMeter::OnMeterOutput);
 
-		Analyzer->StartAnalyzing(InAudioDeviceId, AudioBus.Get());
+			Analyzer->StartAnalyzing(InAudioDeviceId, AudioBus.Get());
+		}
 
 		constexpr float DefaultMeterValue = -60.0f;
 		constexpr float DefaultPeakValue = -60.0f;
@@ -333,12 +337,12 @@ namespace AudioWidgets
 		if (Params.StyleSet->HasWidgetStyle<FAudioMaterialMeterStyle>(MaterialMeterStyleName))
 		{
 			const FAudioMaterialMeterStyle& AudioMaterialMeterStyle = Params.StyleSet->GetWidgetStyle<FAudioMaterialMeterStyle>(MaterialMeterStyleName);
-			return MakeShared<FAudioMeter>(Params.AudioBusInfo.AudioBus->GetNumChannels(), Params.AudioBusInfo.AudioDeviceId, AudioMaterialMeterStyle, Params.AudioBusInfo.AudioBus);
+			return MakeShared<FAudioMeter>(Params.AudioBusInfo.GetNumChannels(), Params.AudioBusInfo.AudioDeviceId, AudioMaterialMeterStyle, Params.AudioBusInfo.AudioBus);
 		}
 		else
 		{
 			const FAudioMeterDefaultColorStyle& MeterDefaultColorStyle = Params.StyleSet->GetWidgetStyle<FAudioMeterDefaultColorStyle>("AudioMeter.DefaultColorStyle");
-			return MakeShared<FAudioMeter>(Params.AudioBusInfo.AudioBus->GetNumChannels(), Params.AudioBusInfo.AudioDeviceId, Params.AudioBusInfo.AudioBus, &MeterDefaultColorStyle);
+			return MakeShared<FAudioMeter>(Params.AudioBusInfo.GetNumChannels(), Params.AudioBusInfo.AudioDeviceId, Params.AudioBusInfo.AudioBus, &MeterDefaultColorStyle);
 		}
 	}
 } // namespace AudioWidgets
diff --git a/Engine/Plugins/Runtime/AudioWidgets/Source/AudioWidgets/Private/AudioOscilloscope.cpp b/Engine/Plugins/Runtime/AudioWidgets/Source/AudioWidgets/Private/AudioOscilloscope.cpp
--- a/Engine/Plugins/Runtime/AudioWidgets/Source/AudioWidgets/Private/AudioOscilloscope.cpp
+++ b/Engine/Plugins/Runtime/AudioWidgets/Source/AudioWidgets/Private/AudioOscilloscope.cpp
@@ -32,7 +32,7 @@ namespace AudioWidgets
 			ensure(InExternalAudioBus->GetNumChannels() == InNumChannels);
 			AudioBus = TStrongObjectPtr(InExternalAudioBus.Get());
 		}
-		else
+		else if (InNumChannels > 0)
 		{
 			CreateAudioBus(InNumChannels);
 		}
@@ -53,6 +53,13 @@ namespace AudioWidgets
 		const float InAnalysisPeriodMs,
 		const EAudioPanelLayoutType InPanelLayoutType)
 	{
+		if (InAudioDeviceId == FAudioBusInfo::InvalidAudioDeviceId)
+		{
+			// Can't have a data provider without a valid audio device.
+			AudioSamplesDataProvider.Reset();
+			return;
+		}
+
 		check(AudioBus);
 
 		const uint32 NumChannelsToProvide = (InPanelLayoutType == EAudioPanelLayoutType::Advanced) ? 1 : AudioBus->GetNumChannels(); // Advanced mode waveform display is based on channel selection
@@ -61,9 +68,13 @@ namespace AudioWidgets
 
 	void FAudioOscilloscope::CreateOscilloscopeWidget(const uint32 InNumChannels, const EAudioPanelLayoutType InPanelLayoutType, const FAudioOscilloscopePanelStyle* InOscilloscopePanelStyle)
 	{
-		check(AudioSamplesDataProvider);
+		FFixedSampledSequenceView SequenceView{ .NumDimensions = FMath::Max(InNumChannels, 1U), .SampleRate = 48000}; // Initialize with usable fallback values (but no sample data in the view).
 
-		const FFixedSampledSequenceView SequenceView = AudioSamplesDataProvider->GetDataView();
+		if (AudioSamplesDataProvider.IsValid())
+		{
+			// Get the actual sequence view from the data provider:
+			SequenceView = AudioSamplesDataProvider->GetDataView();
+		}
 
 		if (InOscilloscopePanelStyle)
 		{
@@ -81,27 +92,36 @@ namespace AudioWidgets
 			OscilloscopePanelWidget->BuildWidget(SequenceView, InNumChannels, InPanelLayoutType);
 		}
 
-		// Interconnect data provider and widget
-		AudioSamplesDataProvider->OnDataViewGenerated.AddSP(OscilloscopePanelWidget.Get(), &SAudioOscilloscopePanelWidget::ReceiveSequenceView);
-
-		if (InPanelLayoutType == EAudioPanelLayoutType::Advanced)
+		if (AudioSamplesDataProvider.IsValid())
 		{
-			OscilloscopePanelWidget->OnSelectedChannelChanged.AddSP(AudioSamplesDataProvider.Get(),  &FWaveformAudioSamplesDataProvider::SetChannelToAnalyze);
-			OscilloscopePanelWidget->OnTriggerModeChanged.AddSP(AudioSamplesDataProvider.Get(),      &FWaveformAudioSamplesDataProvider::SetTriggerMode);
-			OscilloscopePanelWidget->OnTriggerThresholdChanged.AddSP(AudioSamplesDataProvider.Get(), &FWaveformAudioSamplesDataProvider::SetTriggerThreshold);
-			OscilloscopePanelWidget->OnTimeWindowValueChanged.AddSP(AudioSamplesDataProvider.Get(),  &FWaveformAudioSamplesDataProvider::SetTimeWindow);
-			OscilloscopePanelWidget->OnAnalysisPeriodChanged.AddSP(AudioSamplesDataProvider.Get(),   &FWaveformAudioSamplesDataProvider::SetAnalysisPeriod);
+			// Interconnect data provider and widget
+			AudioSamplesDataProvider->OnDataViewGenerated.AddSP(OscilloscopePanelWidget.Get(), &SAudioOscilloscopePanelWidget::ReceiveSequenceView);
+
+			if (InPanelLayoutType == EAudioPanelLayoutType::Advanced)
+			{
+				OscilloscopePanelWidget->OnSelectedChannelChanged.AddSP(AudioSamplesDataProvider.Get(), &FWaveformAudioSamplesDataProvider::SetChannelToAnalyze);
+				OscilloscopePanelWidget->OnTriggerModeChanged.AddSP(AudioSamplesDataProvider.Get(), &FWaveformAudioSamplesDataProvider::SetTriggerMode);
+				OscilloscopePanelWidget->OnTriggerThresholdChanged.AddSP(AudioSamplesDataProvider.Get(), &FWaveformAudioSamplesDataProvider::SetTriggerThreshold);
+				OscilloscopePanelWidget->OnTimeWindowValueChanged.AddSP(AudioSamplesDataProvider.Get(), &FWaveformAudioSamplesDataProvider::SetTimeWindow);
+				OscilloscopePanelWidget->OnAnalysisPeriodChanged.AddSP(AudioSamplesDataProvider.Get(), &FWaveformAudioSamplesDataProvider::SetAnalysisPeriod);
+			}
 		}
 	}
 
 	void FAudioOscilloscope::StartProcessing()
 	{
-		AudioSamplesDataProvider->StartProcessing();
+		if (AudioSamplesDataProvider.IsValid())
+		{
+			AudioSamplesDataProvider->StartProcessing();
+		}
 	}
 
 	void FAudioOscilloscope::StopProcessing()
 	{
-		AudioSamplesDataProvider->StopProcessing();
+		if (AudioSamplesDataProvider.IsValid())
+		{
+			AudioSamplesDataProvider->StopProcessing();
+		}
 	}
 
 	UAudioBus* FAudioOscilloscope::GetAudioBus() const
@@ -135,7 +155,7 @@ namespace AudioWidgets
 	{
 		return MakeShared<FAudioOscilloscope>(
 			Params.AudioBusInfo.AudioDeviceId,
-			Params.AudioBusInfo.AudioBus->GetNumChannels(),
+			Params.AudioBusInfo.GetNumChannels(),
 			RackUnitTimeWindowMs,
 			RackUnitMaxTimeWindowMs,
 			RackUnitAnalysisPeriodMs,
diff --git a/Engine/Plugins/Runtime/AudioWidgets/Source/AudioWidgets/Private/AudioSpectrogram.cpp b/Engine/Plugins/Runtime/AudioWidgets/Source/AudioWidgets/Private/AudioSpectrogram.cpp
--- a/Engine/Plugins/Runtime/AudioWidgets/Source/AudioWidgets/Private/AudioSpectrogram.cpp
+++ b/Engine/Plugins/Runtime/AudioWidgets/Source/AudioWidgets/Private/AudioSpectrogram.cpp
@@ -124,19 +124,24 @@ namespace AudioWidgets
 
 	void FAudioSpectrogram::Init(int32 InNumChannels, Audio::FDeviceId InAudioDeviceId, TObjectPtr<UAudioBus> InExternalAudioBus)
 	{
-		check(InNumChannels > 0);
-
 		Teardown();
 
 		AudioDeviceId = InAudioDeviceId;
-		bUseExternalAudioBus = InExternalAudioBus != nullptr;
-		AudioBus = bUseExternalAudioBus ? TStrongObjectPtr(InExternalAudioBus.Get()) : TStrongObjectPtr(NewObject<UAudioBus>());
-		AudioBus->AudioBusChannels = EAudioBusChannels(InNumChannels - 1);
 
-		CreateSynesthesiaSpectrumAnalyzer();
-		CreateConstantQAnalyzer();
+		// Only create analyzers etc if we have an audio device:
+		if (AudioDeviceId != FAudioBusInfo::InvalidAudioDeviceId)
+		{
+			check(InNumChannels > 0);
+
+			bUseExternalAudioBus = InExternalAudioBus != nullptr;
+			AudioBus = bUseExternalAudioBus ? TStrongObjectPtr(InExternalAudioBus.Get()) : TStrongObjectPtr(NewObject<UAudioBus>());
+			AudioBus->AudioBusChannels = EAudioBusChannels(InNumChannels - 1);
+
+			CreateSynesthesiaSpectrumAnalyzer();
+			CreateConstantQAnalyzer();
 
-		StartAnalyzing(AnalyzerType.Get());
+			StartAnalyzing(AnalyzerType.Get());
+		}
 	}
 
 	void FAudioSpectrogram::StartAnalyzing(const EAudioSpectrumAnalyzerType InAnalyzerType)
@@ -364,6 +369,13 @@ namespace AudioWidgets
 
 	EActiveTimerReturnType FAudioSpectrogram::Update(double InCurrentTime, float InDeltaTime)
 	{
+		if (AudioDeviceId == FAudioBusInfo::InvalidAudioDeviceId)
+		{
+			// No analyzers available if no valid audio device.
+			ensure(!ActiveAnalyzerType.IsSet());
+			return EActiveTimerReturnType::Continue;
+		}
+
 		const EAudioSpectrumAnalyzerType RequiredAnalyzerType = AnalyzerType.Get();
 		const EFFTSize FFTAnalyzerRequiredFFTSize = FFTAnalyzerFFTSize.Get();
 		const EConstantQFFTSizeEnum CQTAnalyzerRequiredFFTSize = CQTAnalyzerFFTSize.Get();
@@ -449,7 +461,7 @@ namespace AudioWidgets
 		using namespace AudioSpectrogramPrivate;
 
 		FAudioSpectrogramParams AnalyzerParams;
-		AnalyzerParams.NumChannels = Params.AudioBusInfo.AudioBus->GetNumChannels();
+		AnalyzerParams.NumChannels = Params.AudioBusInfo.GetNumChannels();
 		AnalyzerParams.AudioDeviceId = Params.AudioBusInfo.AudioDeviceId;
 		AnalyzerParams.ExternalAudioBus = Params.AudioBusInfo.AudioBus;
 
diff --git a/Engine/Plugins/Runtime/AudioWidgets/Source/AudioWidgets/Private/AudioSpectrumAnalyzer.cpp b/Engine/Plugins/Runtime/AudioWidgets/Source/AudioWidgets/Private/AudioSpectrumAnalyzer.cpp
--- a/Engine/Plugins/Runtime/AudioWidgets/Source/AudioWidgets/Private/AudioSpectrumAnalyzer.cpp
+++ b/Engine/Plugins/Runtime/AudioWidgets/Source/AudioWidgets/Private/AudioSpectrumAnalyzer.cpp
@@ -160,19 +160,24 @@ namespace AudioWidgets
 
 	void FAudioSpectrumAnalyzer::Init(int32 InNumChannels, Audio::FDeviceId InAudioDeviceId, TObjectPtr<UAudioBus> InExternalAudioBus)
 	{
-		check(InNumChannels > 0);
-
 		Teardown();
 
 		AudioDeviceId = InAudioDeviceId;
-		bUseExternalAudioBus = InExternalAudioBus != nullptr;
-		AudioBus = bUseExternalAudioBus ? TStrongObjectPtr(InExternalAudioBus.Get()) : TStrongObjectPtr(NewObject<UAudioBus>());
-		AudioBus->AudioBusChannels = EAudioBusChannels(InNumChannels - 1);
 
-		CreateSynesthesiaSpectrumAnalyzer();
-		CreateConstantQAnalyzer();
+		// Only create analyzers etc if we have an audio device:
+		if (AudioDeviceId != FAudioBusInfo::InvalidAudioDeviceId)
+		{
+			check(InNumChannels > 0);
+
+			bUseExternalAudioBus = InExternalAudioBus != nullptr;
+			AudioBus = bUseExternalAudioBus ? TStrongObjectPtr(InExternalAudioBus.Get()) : TStrongObjectPtr(NewObject<UAudioBus>());
+			AudioBus->AudioBusChannels = EAudioBusChannels(InNumChannels - 1);
+
+			CreateSynesthesiaSpectrumAnalyzer();
+			CreateConstantQAnalyzer();
 
-		StartAnalyzing(AnalyzerType.Get());
+			StartAnalyzing(AnalyzerType.Get());
+		}
 	}
 
 	void FAudioSpectrumAnalyzer::StartAnalyzing(const EAudioSpectrumAnalyzerType InAnalyzerType)
@@ -513,6 +518,13 @@ namespace AudioWidgets
 
 	void FAudioSpectrumAnalyzer::UpdateAnalyzerSettings()
 	{
+		if (AudioDeviceId == FAudioBusInfo::InvalidAudioDeviceId)
+		{
+			// No analyzers available if no valid audio device.
+			ensure(!ActiveAnalyzerType.IsSet());
+			return;
+		}
+
 		const EAudioSpectrumAnalyzerType RequiredAnalyzerType = AnalyzerType.Get();
 		const EFFTSize FFTAnalyzerRequiredFFTSize = FFTAnalyzerFFTSize.Get();
 		const EConstantQFFTSizeEnum CQTAnalyzerRequiredFFTSize = CQTAnalyzerFFTSize.Get();
@@ -604,7 +616,7 @@ namespace AudioWidgets
 		using namespace AudioSpectrumAnalyzerPrivate;
 
 		FAudioSpectrumAnalyzerParams AnalyzerParams;
-		AnalyzerParams.NumChannels = Params.AudioBusInfo.AudioBus->GetNumChannels();
+		AnalyzerParams.NumChannels = Params.AudioBusInfo.GetNumChannels();
 		AnalyzerParams.AudioDeviceId = Params.AudioBusInfo.AudioDeviceId;
 		AnalyzerParams.ExternalAudioBus = Params.AudioBusInfo.AudioBus;
 
diff --git a/Engine/Plugins/Runtime/AudioWidgets/Source/AudioWidgets/Private/AudioVectorscope.cpp b/Engine/Plugins/Runtime/AudioWidgets/Source/AudioWidgets/Private/AudioVectorscope.cpp
--- a/Engine/Plugins/Runtime/AudioWidgets/Source/AudioWidgets/Private/AudioVectorscope.cpp
+++ b/Engine/Plugins/Runtime/AudioWidgets/Source/AudioWidgets/Private/AudioVectorscope.cpp
@@ -33,7 +33,7 @@ namespace AudioWidgets
 			ensure(InExternalAudioBus->GetNumChannels() == InNumChannels);
 			AudioBus = TStrongObjectPtr(InExternalAudioBus.Get());
 		}
-		else
+		else if (InNumChannels > 0)
 		{
 			CreateAudioBus(InNumChannels);
 		}
@@ -50,6 +50,13 @@ namespace AudioWidgets
 
 	void FAudioVectorscope::CreateDataProvider(Audio::FDeviceId InAudioDeviceId, const float InTimeWindowMs,	const float InMaxTimeWindowMs, const float InAnalysisPeriodMs)
 	{
+		if (InAudioDeviceId == FAudioBusInfo::InvalidAudioDeviceId)
+		{
+			// Can't have a data provider without a valid audio device.
+			AudioSamplesDataProvider.Reset();
+			return;
+		}
+
 		check(AudioBus);
 
 		AudioSamplesDataProvider = MakeShared<FWaveformAudioSamplesDataProvider>(InAudioDeviceId, AudioBus.Get(), AudioBus->GetNumChannels(), InTimeWindowMs, InMaxTimeWindowMs, InAnalysisPeriodMs);
@@ -57,9 +64,13 @@ namespace AudioWidgets
 
 	void FAudioVectorscope::CreateVectorscopeWidget(const EAudioPanelLayoutType InPanelLayoutType, const FAudioVectorscopePanelStyle* PanelStyle)
 	{
-		check(AudioSamplesDataProvider);
+		FFixedSampledSequenceView SequenceView{ .NumDimensions = 2, .SampleRate = 48000 }; // Initialize with usable fallback values (but no sample data in the view).
 
-		const FFixedSampledSequenceView SequenceView = AudioSamplesDataProvider->GetDataView();
+		if (AudioSamplesDataProvider.IsValid())
+		{
+			// Get the actual sequence view from the data provider:
+			SequenceView = AudioSamplesDataProvider->GetDataView();
+		}
 
 		if (PanelStyle)
 		{
@@ -77,23 +88,32 @@ namespace AudioWidgets
 			VectorscopePanelWidget->BuildWidget(SequenceView, InPanelLayoutType);
 		}
 
-		// Interconnect data provider and widget
-		AudioSamplesDataProvider->OnDataViewGenerated.AddSP(VectorscopePanelWidget.Get(), &SAudioVectorscopePanelWidget::ReceiveSequenceView);
-
-		if (InPanelLayoutType == EAudioPanelLayoutType::Advanced)
+		if (AudioSamplesDataProvider.IsValid())
 		{
-			VectorscopePanelWidget->OnDisplayPersistenceValueChanged.AddSP(AudioSamplesDataProvider.Get(), &FWaveformAudioSamplesDataProvider::SetTimeWindow);
+			// Interconnect data provider and widget
+			AudioSamplesDataProvider->OnDataViewGenerated.AddSP(VectorscopePanelWidget.Get(), &SAudioVectorscopePanelWidget::ReceiveSequenceView);
+
+			if (InPanelLayoutType == EAudioPanelLayoutType::Advanced)
+			{
+				VectorscopePanelWidget->OnDisplayPersistenceValueChanged.AddSP(AudioSamplesDataProvider.Get(), &FWaveformAudioSamplesDataProvider::SetTimeWindow);
+			}
 		}
 	}
 
 	void FAudioVectorscope::StartProcessing()
 	{
-		AudioSamplesDataProvider->StartProcessing();
+		if (AudioSamplesDataProvider.IsValid())
+		{
+			AudioSamplesDataProvider->StartProcessing();
+		}
 	}
 
 	void FAudioVectorscope::StopProcessing()
 	{
-		AudioSamplesDataProvider->StopProcessing();
+		if (AudioSamplesDataProvider.IsValid())
+		{
+			AudioSamplesDataProvider->StopProcessing();
+		}
 	}
 
 	UAudioBus* FAudioVectorscope::GetAudioBus() const
@@ -127,7 +147,7 @@ namespace AudioWidgets
 	{
 		return MakeShared<FAudioVectorscope>(
 			Params.AudioBusInfo.AudioDeviceId,
-			Params.AudioBusInfo.AudioBus->GetNumChannels(),
+			Params.AudioBusInfo.GetNumChannels(),
 			RackUnitTimeWindowMs,
 			RackUnitMaxTimeWindowMs,
 			RackUnitAnalysisPeriodMs,
diff --git a/Engine/Plugins/Runtime/AudioWidgets/Source/AudioWidgets/Private/SAudioOscilloscopePanelWidget.cpp b/Engine/Plugins/Runtime/AudioWidgets/Source/AudioWidgets/Private/SAudioOscilloscopePanelWidget.cpp
--- a/Engine/Plugins/Runtime/AudioWidgets/Source/AudioWidgets/Private/SAudioOscilloscopePanelWidget.cpp
+++ b/Engine/Plugins/Runtime/AudioWidgets/Source/AudioWidgets/Private/SAudioOscilloscopePanelWidget.cpp
@@ -309,7 +309,9 @@ void SAudioOscilloscopePanelWidget::CreateLayout()
 
 void SAudioOscilloscopePanelWidget::CreateGridData(const FFixedSampleSequenceRulerStyle& RulerStyle)
 {
-	SequenceGridData = MakeShared<FFixedSampledSequenceGridData>(DataView.SampleData.Num() / DataView.NumDimensions, DataView.SampleRate, RulerStyle.TicksTextFont, RulerStyle.DesiredWidth);
+	// Calculate number of frames, ensuring to span at least two frames so that FFixedSampledSequenceGridData can successfully calculate its grid metrics:
+	const uint32 TotalFrames = FMath::Max(DataView.SampleData.Num() / DataView.NumDimensions, 2U);
+	SequenceGridData = MakeShared<FFixedSampledSequenceGridData>(TotalFrames, DataView.SampleRate, RulerStyle.TicksTextFont, RulerStyle.DesiredWidth);
 }
 
 void SAudioOscilloscopePanelWidget::CreateSequenceRuler(TSharedRef<FFixedSampledSequenceGridData> InGridData, const FFixedSampleSequenceRulerStyle& RulerStyle)
diff --git a/Engine/Plugins/Runtime/AudioWidgets/Source/AudioWidgets/Public/AudioAnalyzerRack.h b/Engine/Plugins/Runtime/AudioWidgets/Source/AudioWidgets/Public/AudioAnalyzerRack.h
--- a/Engine/Plugins/Runtime/AudioWidgets/Source/AudioWidgets/Public/AudioAnalyzerRack.h
+++ b/Engine/Plugins/Runtime/AudioWidgets/Source/AudioWidgets/Public/AudioAnalyzerRack.h
@@ -14,8 +14,12 @@ namespace AudioWidgets
 	 */
 	struct FAudioBusInfo
 	{
-		Audio::FDeviceId AudioDeviceId = static_cast<Audio::FDeviceId>(INDEX_NONE);
+		static constexpr Audio::FDeviceId InvalidAudioDeviceId = static_cast<Audio::FDeviceId>(INDEX_NONE);
+
+		Audio::FDeviceId AudioDeviceId = InvalidAudioDeviceId;
 		TObjectPtr<UAudioBus> AudioBus = nullptr;
+
+		int32 GetNumChannels() const { return (AudioBus) ? AudioBus->GetNumChannels() : 0; }
 	};
 
 	/**
