Build warning fix  - Remove LoadObject in audio material widget style struct constructors

diff --git a/Engine/Plugins/Runtime/AudioWidgets/Source/AudioWidgets/Private/AudioMaterialSlate/AudioMaterialSlateTypes.cpp b/Engine/Plugins/Runtime/AudioWidgets/Source/AudioWidgets/Private/AudioMaterialSlate/AudioMaterialSlateTypes.cpp
--- a/Engine/Plugins/Runtime/AudioWidgets/Source/AudioWidgets/Private/AudioMaterialSlate/AudioMaterialSlateTypes.cpp
+++ b/Engine/Plugins/Runtime/AudioWidgets/Source/AudioWidgets/Private/AudioMaterialSlate/AudioMaterialSlateTypes.cpp
@@ -86,9 +86,6 @@ FAudioMaterialButtonStyle::FAudioMaterialButtonStyle()
 	, ButtonUnpressedOutlineColor(SlateTypesPrivate::ButtonUnpressedOutlineColor)
 	, ButtonPressedOutlineColor(SlateTypesPrivate::ButtonPressedOutlineColor)
 {
-	FString Path = PLUGIN_BASE_DIR + "MI_AudioMaterialToggleButton.MI_AudioMaterialToggleButton";
-	Material = LoadObject<UMaterialInterface>(nullptr, *Path);
-
 	DesiredSize = FVector2f(128.f, 128.f);	
 }
 
@@ -104,6 +101,17 @@ const FAudioMaterialButtonStyle& FAudioMaterialButtonStyle::GetDefault()
 	return Default;
 }
 
+UMaterialInstanceDynamic* FAudioMaterialButtonStyle::CreateDynamicMaterial(UObject* InOuter) const
+{
+	// Use default material if none provided
+	if (!Material)
+	{
+		const FString Path = PLUGIN_BASE_DIR + "MI_AudioMaterialToggleButton.MI_AudioMaterialToggleButton";
+		return UMaterialInstanceDynamic::Create(LoadObject<UMaterialInterface>(nullptr, *Path), InOuter);
+	}
+	return Super::CreateDynamicMaterial(InOuter);
+}
+
 FAudioMaterialSliderStyle::FAudioMaterialSliderStyle()
 	: SliderBackgroundColor(SlateTypesPrivate::SliderBackgroundColor)
 	, SliderBackgroundAccentColor(SlateTypesPrivate::SliderBackgroundAccentColor)
@@ -112,9 +120,6 @@ FAudioMaterialSliderStyle::FAudioMaterialSliderStyle()
 	, SliderHandleOutlineColor(SlateTypesPrivate::SliderHandleOutlineColor)
 	, TextBoxStyle(FAudioTextBoxStyle::GetDefault())
 {
-	FString Path = PLUGIN_BASE_DIR + "MI_AudioMaterialRoundedSlider.MI_AudioMaterialRoundedSlider";
-	Material = LoadObject<UMaterialInterface>(nullptr, *Path);
-
 	DesiredSize = FVector2f(25.f, 250.f);
 }
 
@@ -131,6 +136,17 @@ const FAudioMaterialSliderStyle& FAudioMaterialSliderStyle::GetDefault()
 	return Default;
 }
 
+UMaterialInstanceDynamic* FAudioMaterialSliderStyle::CreateDynamicMaterial(UObject* InOuter) const
+{
+	// Use default material if none provided
+	if (!Material)
+	{
+		const FString Path = PLUGIN_BASE_DIR + "MI_AudioMaterialRoundedSlider.MI_AudioMaterialRoundedSlider";
+		return UMaterialInstanceDynamic::Create(LoadObject<UMaterialInterface>(nullptr, *Path), InOuter);
+	}
+	return Super::CreateDynamicMaterial(InOuter);
+}
+
 FAudioMaterialKnobStyle::FAudioMaterialKnobStyle()
 	: KnobMainColor(SlateTypesPrivate::KnobMainColor)
 	, KnobAccentColor(SlateTypesPrivate::KnobAccentColor)
@@ -145,9 +161,6 @@ FAudioMaterialKnobStyle::FAudioMaterialKnobStyle()
 	, KnobBarFillMaxColor(SlateTypesPrivate::KnobBarFillMaxColor)
 	, KnobBarFillTintColor(SlateTypesPrivate::KnobBarFillTintColor)
 {
-	FString Path = PLUGIN_BASE_DIR + "MI_AudioMaterialKnob.MI_AudioMaterialKnob";
-	Material = LoadObject<UMaterialInterface>(nullptr, *Path);
-
 	DesiredSize = FVector2f(128.f,128.f);
 }
 
@@ -164,6 +177,17 @@ void FAudioMaterialKnobStyle::GetResources(TArray<const FSlateBrush*>& OutBrushe
 	// Add any brush resources here so that Slate can correctly atlas and reference them
 }
 
+UMaterialInstanceDynamic* FAudioMaterialKnobStyle::CreateDynamicMaterial(UObject* InOuter) const
+{
+	// Use default material if none provided
+	if (!Material)
+	{
+		const FString Path = PLUGIN_BASE_DIR + "MI_AudioMaterialKnob.MI_AudioMaterialKnob";
+		return UMaterialInstanceDynamic::Create(LoadObject<UMaterialInterface>(nullptr, *Path), InOuter);
+	}
+	return Super::CreateDynamicMaterial(InOuter);
+}
+
 FAudioMaterialMeterStyle::FAudioMaterialMeterStyle()
 	: MeterFillMinColor(SlateTypesPrivate::MeterFillMinColor)
 	, MeterFillMidColor(SlateTypesPrivate::MeterFillMidColor)
@@ -179,9 +203,6 @@ FAudioMaterialMeterStyle::FAudioMaterialMeterStyle()
 	, DecibelsPerHash(SlateTypesPrivate::MeterDecibelsPerHash)
 	, Font(FStyleDefaults::GetFontInfo(5))
 {
-	FString Path = PLUGIN_BASE_DIR + "MI_AudioMaterialMeter.MI_AudioMaterialMeter";
-	Material = LoadObject<UMaterialInterface>(nullptr, *Path);
-
 	DesiredSize = FVector2f(25.f, 512.f);
 }
 
@@ -198,6 +219,17 @@ void FAudioMaterialMeterStyle::GetResources(TArray<const FSlateBrush*>& OutBrush
 	// Add any brush resources here so that Slate can correctly atlas and reference them
 }
 
+UMaterialInstanceDynamic* FAudioMaterialMeterStyle::CreateDynamicMaterial(UObject* InOuter) const
+{
+	// Use default material if none provided
+	if (!Material)
+	{
+		const FString Path = PLUGIN_BASE_DIR + "MI_AudioMaterialMeter.MI_AudioMaterialMeter";
+		return UMaterialInstanceDynamic::Create(LoadObject<UMaterialInterface>(nullptr, *Path), InOuter);
+	}
+	return Super::CreateDynamicMaterial(InOuter);
+}
+
 PRAGMA_DISABLE_DEPRECATION_WARNINGS
 FAudioMaterialMeterStyle& FAudioMaterialMeterStyle::SetFont(const FName& InFontName, uint16 InSize) { Font = FSlateFontInfo(InFontName, InSize); return *this; }
 FAudioMaterialMeterStyle& FAudioMaterialMeterStyle::SetFont(const FString& InFontName, uint16 InSize) { Font = FSlateFontInfo(*InFontName, InSize); return *this; }
@@ -214,9 +246,6 @@ FAudioMaterialEnvelopeStyle::FAudioMaterialEnvelopeStyle()
 	, BackgroundColor(FLinearColor::Black)
 	, OutlineColor(FLinearColor::Gray)
 {
-	FString Path = PLUGIN_BASE_DIR + "MI_AudioMaterialEnvelope_ADSR.MI_AudioMaterialEnvelope_ADSR";
-	Material = LoadObject<UMaterialInterface>(nullptr, *Path);
-
 	DesiredSize = FVector2f(256.f, 256.f);
 }
 
@@ -232,3 +261,14 @@ void FAudioMaterialEnvelopeStyle::GetResources(TArray<const FSlateBrush*>& OutBr
 {
 	// Add any brush resources here so that Slate can correctly atlas and reference them
 }
+
+UMaterialInstanceDynamic* FAudioMaterialEnvelopeStyle::CreateDynamicMaterial(UObject* InOuter) const
+{
+	// Use default material if none provided
+	if (!Material)
+	{
+		const FString Path = PLUGIN_BASE_DIR + "MI_AudioMaterialEnvelope_ADSR.MI_AudioMaterialEnvelope_ADSR";
+		return UMaterialInstanceDynamic::Create(LoadObject<UMaterialInterface>(nullptr, *Path), InOuter);
+	}
+	return Super::CreateDynamicMaterial(InOuter);
+}
diff --git a/Engine/Plugins/Runtime/AudioWidgets/Source/AudioWidgets/Public/AudioMaterialSlate/AudioMaterialSlateTypes.h b/Engine/Plugins/Runtime/AudioWidgets/Source/AudioWidgets/Public/AudioMaterialSlate/AudioMaterialSlateTypes.h
--- a/Engine/Plugins/Runtime/AudioWidgets/Source/AudioWidgets/Public/AudioMaterialSlate/AudioMaterialSlateTypes.h
+++ b/Engine/Plugins/Runtime/AudioWidgets/Source/AudioWidgets/Public/AudioMaterialSlate/AudioMaterialSlateTypes.h
@@ -32,7 +32,7 @@ public:
 
 public:
 
-	UMaterialInstanceDynamic* CreateDynamicMaterial(UObject* InOuter) const;
+	virtual UMaterialInstanceDynamic* CreateDynamicMaterial(UObject* InOuter) const;
 
 };
 
@@ -51,6 +51,7 @@ struct AUDIOWIDGETS_API FAudioMaterialButtonStyle : public FAudioMaterialWidgetS
 	static const FName TypeName;
 	virtual const FName GetTypeName() const override { return TypeName; };
 	static const FAudioMaterialButtonStyle& GetDefault();
+	virtual UMaterialInstanceDynamic* CreateDynamicMaterial(UObject* InOuter) const override;
 
 	FAudioMaterialButtonStyle& SetMaterial(UMaterialInterface* InMaterialInterface) { Material = InMaterialInterface; return *this; }
 
@@ -106,6 +107,7 @@ struct AUDIOWIDGETS_API FAudioMaterialSliderStyle : public FAudioMaterialWidgetS
 	static const FName TypeName;
 	virtual const FName GetTypeName() const override { return TypeName; };
 	static const FAudioMaterialSliderStyle& GetDefault();
+	virtual UMaterialInstanceDynamic* CreateDynamicMaterial(UObject* InOuter) const override;
 
 	FAudioMaterialSliderStyle& SetMaterial(UMaterialInterface* InMaterialInterface) { Material = InMaterialInterface; return *this; }
 
@@ -156,6 +158,7 @@ struct AUDIOWIDGETS_API FAudioMaterialKnobStyle : public FAudioMaterialWidgetSty
 	static const FName TypeName;
 	virtual const FName GetTypeName() const override { return TypeName; };
 	static const FAudioMaterialKnobStyle& GetDefault();
+	virtual UMaterialInstanceDynamic* CreateDynamicMaterial(UObject* InOuter) const override;
 
 	FAudioMaterialKnobStyle& SetMaterial(UMaterialInterface* InMaterialInterface) { Material = InMaterialInterface; return *this; }
 
@@ -241,6 +244,7 @@ struct AUDIOWIDGETS_API FAudioMaterialMeterStyle : public FAudioMaterialWidgetSt
 	static const FName TypeName;
 	virtual const FName GetTypeName() const override { return TypeName; };
 	static const FAudioMaterialMeterStyle& GetDefault();
+	virtual UMaterialInstanceDynamic* CreateDynamicMaterial(UObject* InOuter) const override;
 
 	FAudioMaterialMeterStyle& SetMaterial(UMaterialInterface* InMaterialInterface) { Material = InMaterialInterface; return *this; }
 
@@ -348,6 +352,7 @@ struct AUDIOWIDGETS_API FAudioMaterialEnvelopeStyle : public FAudioMaterialWidge
 	static const FName TypeName;
 	virtual const FName GetTypeName() const override { return TypeName; };
 	static const FAudioMaterialEnvelopeStyle& GetDefault();
+	virtual UMaterialInstanceDynamic* CreateDynamicMaterial(UObject* InOuter) const override;
 
 	FAudioMaterialEnvelopeStyle& SetMaterial(UMaterialInterface* InMaterialInterface) { Material = InMaterialInterface; return *this; }
 
