[Audio Insights] Lower meter sensitivity threshold to match activity indicator

diff --git a/Engine/Plugins/Runtime/AudioWidgets/Source/AudioWidgets/Private/AudioMeterStyle.cpp b/Engine/Plugins/Runtime/AudioWidgets/Source/AudioWidgets/Private/AudioMeterStyle.cpp
--- a/Engine/Plugins/Runtime/AudioWidgets/Source/AudioWidgets/Private/AudioMeterStyle.cpp
+++ b/Engine/Plugins/Runtime/AudioWidgets/Source/AudioWidgets/Private/AudioMeterStyle.cpp
@@ -10,13 +10,13 @@ FAudioMeterStyle::FAudioMeterStyle()
 	, MeterPadding(FVector2D(10.0f, 5.0f))
 	, MeterValuePadding(3.0f)
 	, PeakValueWidth(2.0f)
-	, ValueRangeDb(FVector2D(-60, 10))
+	, ValueRangeDb(FVector2D(-160, 10))
 	, bShowScale(true)
 	, bScaleSide(true)
 	, ScaleHashOffset(5.0f)
 	, ScaleHashWidth(1.0f)
 	, ScaleHashHeight(10.0f)
-	, DecibelsPerHash(5)
+	, DecibelsPerHash(10)
 	, Font(FStyleDefaults::GetFontInfo(5))
 {
 }
