Use modular feature to instantiate AudioPropertiesSheet details injector

diff --git a/Engine/Plugins/Runtime/AudioWidgets/Source/AudioWidgetsEditor/Private/SoundBaseDetails.cpp b/Engine/Plugins/Runtime/AudioWidgets/Source/AudioWidgetsEditor/Private/SoundBaseDetails.cpp
--- a/Engine/Plugins/Runtime/AudioWidgets/Source/AudioWidgetsEditor/Private/SoundBaseDetails.cpp
+++ b/Engine/Plugins/Runtime/AudioWidgets/Source/AudioWidgetsEditor/Private/SoundBaseDetails.cpp
@@ -4,6 +4,7 @@
 
 #include "IAudioPropertiesDetailsInjector.h"
 #include "DetailLayoutBuilder.h"
+#include "Features/IModularFeatures.h"
 #include "IPropertyUtilities.h"
 #include "Modules/ModuleManager.h"
 #include "PropertyHandle.h"
@@ -44,36 +45,16 @@ void FSoundBaseDetails::CustomizeDetails(IDetailLayoutBuilder& DetailBuilder)
 void FSoundBaseDetails::InjectPropertySheetView(IDetailLayoutBuilder& DetailBuilder)
 {
 	check(FModuleManager::Get().IsModuleLoaded("AudioPropertiesEditor"))
+	
+	IModularFeatures::FScopedLockModularFeatureList ScopedLockModularFeatureList;
+	TArray<IAudioPropertiesDetailsInjectorBuilder*> AvailableDetailsInjectorBuilders = IModularFeatures::Get().GetModularFeatureImplementations<IAudioPropertiesDetailsInjectorBuilder>(AudioPropertiesDetailsInjector::BuilderModularFeatureName);
 
-	FString ModuleFileName = "UnrealEditor-AudioPropertiesEditor.dll";
-
-#if !IS_MONOLITHIC
-	ModuleFileName = FModuleManager::Get().GetModuleFilename("AudioPropertiesEditor");
-#endif
-
-	void* AudioPropertiesEditorModuleDllHandle = FPlatformProcess::GetDllHandle(*ModuleFileName);
-
-	if (!AudioPropertiesEditorModuleDllHandle)
-	{
-		return;
-	}
-
-	typedef IAudioPropertiesDetailsInjector* (*CreateInjectorFunc)();
-	CreateInjectorFunc CreateInjector = (CreateInjectorFunc)FPlatformProcess::GetDllExport(AudioPropertiesEditorModuleDllHandle, TEXT("CreateAudioPropertiesDetailsInjectorExtern"));
-
-	if (!CreateInjector)
-	{
-		return;
-	}
-
-	IAudioPropertiesDetailsInjector* RawInjector = CreateInjector();
-	if (RawInjector != nullptr)
+	if(IAudioPropertiesDetailsInjectorBuilder* Builder = AvailableDetailsInjectorBuilders[0])
 	{
-		AudioPropertiesInjector = TSharedPtr<IAudioPropertiesDetailsInjector>(RawInjector);
+		AudioPropertiesInjector = TSharedPtr<IAudioPropertiesDetailsInjector>(Builder->CreateAudioPropertiesDetailsInjector());
 	}
 	
-
-	if (!AudioPropertiesInjector.IsValid())
+	if (!AudioPropertiesInjector || !AudioPropertiesInjector.IsValid())
 	{
 		return;
 	}
diff --git a/Engine/Source/Runtime/AudioExtensions/Private/IAudioPropertiesDetailsInjector.cpp b/Engine/Source/Runtime/AudioExtensions/Private/IAudioPropertiesDetailsInjector.cpp
--- a/Engine/Source/Runtime/AudioExtensions/Private/IAudioPropertiesDetailsInjector.cpp
+++ b/Engine/Source/Runtime/AudioExtensions/Private/IAudioPropertiesDetailsInjector.cpp
@@ -2,4 +2,6 @@
 
 #include "IAudioPropertiesDetailsInjector.h"
 
-IAudioPropertiesDetailsInjector::~IAudioPropertiesDetailsInjector() {};
\ No newline at end of file
+IAudioPropertiesDetailsInjector::~IAudioPropertiesDetailsInjector() {};
+
+IAudioPropertiesDetailsInjectorBuilder::~IAudioPropertiesDetailsInjectorBuilder() {};
diff --git a/Engine/Source/Runtime/AudioExtensions/Public/IAudioPropertiesDetailsInjector.h b/Engine/Source/Runtime/AudioExtensions/Public/IAudioPropertiesDetailsInjector.h
--- a/Engine/Source/Runtime/AudioExtensions/Public/IAudioPropertiesDetailsInjector.h
+++ b/Engine/Source/Runtime/AudioExtensions/Public/IAudioPropertiesDetailsInjector.h
@@ -2,17 +2,32 @@
 
 #pragma once
 
+#include "Features/IModularFeature.h"
 #include "Templates/SharedPointer.h"
 
 class IDetailLayoutBuilder;
 class IPropertyHandle;
 
+namespace AudioPropertiesDetailsInjector
+{
+	static const FLazyName BuilderModularFeatureName = TEXT("AudioPropertiesDetailsInjectorBuilder");
+	
+}
+
 class AUDIOEXTENSIONS_API IAudioPropertiesDetailsInjector : public TSharedFromThis<IAudioPropertiesDetailsInjector>
 {
 public:
-	virtual ~IAudioPropertiesDetailsInjector() = 0;
+	virtual ~IAudioPropertiesDetailsInjector();
 
 	virtual void CustomizeInjectedPropertiesDetails(IDetailLayoutBuilder& DetailBuilder, TSharedRef<IPropertyHandle> PropertySheetPropertyHandle) = 0;
 	virtual void BindDetailCustomizationToPropertySheetChanges(IDetailLayoutBuilder& DetailBuilder, TSharedRef<IPropertyHandle> PropertySheetPropertyHandle) = 0;
 	virtual void UnbindFromPropertySheetChanges() = 0;
+};
+
+class AUDIOEXTENSIONS_API IAudioPropertiesDetailsInjectorBuilder : public TSharedFromThis<IAudioPropertiesDetailsInjectorBuilder>, public IModularFeature
+{
+public:
+	virtual ~IAudioPropertiesDetailsInjectorBuilder();
+
+    virtual IAudioPropertiesDetailsInjector* CreateAudioPropertiesDetailsInjector() = 0;
 };
\ No newline at end of file
