[AudioWidgets] AudioMeter: Adjusted DefaultMeterValue and DefaultPeakValue so the meter level doesn't jump when initialized (complement to the changes done in CL 41062388)

diff --git a/Engine/Plugins/Runtime/AudioWidgets/Source/AudioWidgets/Private/AudioMeter.cpp b/Engine/Plugins/Runtime/AudioWidgets/Source/AudioWidgets/Private/AudioMeter.cpp
--- a/Engine/Plugins/Runtime/AudioWidgets/Source/AudioWidgets/Private/AudioMeter.cpp
+++ b/Engine/Plugins/Runtime/AudioWidgets/Source/AudioWidgets/Private/AudioMeter.cpp
@@ -259,8 +259,8 @@ namespace AudioWidgets
 			Analyzer->StartAnalyzing(InAudioDeviceId, AudioBus.Get());
 		}
 
-		constexpr float DefaultMeterValue = -60.0f;
-		constexpr float DefaultPeakValue = -60.0f;
+		constexpr float DefaultMeterValue = -160.0f;
+		constexpr float DefaultPeakValue = -160.0f;
 		ChannelInfo.Init(FMeterChannelInfo{ DefaultMeterValue, DefaultPeakValue }, InNumChannels);
 
 		if (Widget.IsValid())
