AudioWidgets: Fix static inline TMulticastDelegate variables

#rnx

diff --git a/Engine/Plugins/Runtime/AudioWidgets/Source/AudioWidgets/Private/AudioWidgetsStyle.cpp b/Engine/Plugins/Runtime/AudioWidgets/Source/AudioWidgets/Private/AudioWidgetsStyle.cpp
--- a/Engine/Plugins/Runtime/AudioWidgets/Source/AudioWidgets/Private/AudioWidgetsStyle.cpp
+++ b/Engine/Plugins/Runtime/AudioWidgets/Source/AudioWidgets/Private/AudioWidgetsStyle.cpp
@@ -18,6 +18,13 @@
 #include "Styling/SlateStyleMacros.h"
 #include "Styling/SlateStyleRegistry.h"
 
+FOnNewTimeRulerStyle		FAudioOscilloscopePanelStyle::OnNewTimeRulerStyle;
+FOnNewValueGridOverlayStyle	FAudioOscilloscopePanelStyle::OnNewValueGridOverlayStyle;
+FOnNewWaveformViewerStyle	FAudioOscilloscopePanelStyle::OnNewWaveformViewerStyle;
+
+FOnNewValueGridOverlayStyle	FAudioVectorscopePanelStyle::OnNewValueGridOverlayStyle;
+FOnNewVectorViewerStyle		FAudioVectorscopePanelStyle::OnNewVectorViewerStyle;
+
 FName FAudioWidgetsStyle::StyleName("AudioWidgetsStyle");
 
 namespace AudioWidgetsStylePrivate
diff --git a/Engine/Plugins/Runtime/AudioWidgets/Source/AudioWidgets/Public/AudioOscilloscopePanelStyle.h b/Engine/Plugins/Runtime/AudioWidgets/Source/AudioWidgets/Public/AudioOscilloscopePanelStyle.h
--- a/Engine/Plugins/Runtime/AudioWidgets/Source/AudioWidgets/Public/AudioOscilloscopePanelStyle.h
+++ b/Engine/Plugins/Runtime/AudioWidgets/Source/AudioWidgets/Public/AudioOscilloscopePanelStyle.h
@@ -81,9 +81,9 @@ struct FAudioOscilloscopePanelStyle : public FSlateWidgetStyle
 
 	inline static FName TypeName = FName("FAudioOscilloscopePanelStyle");
 
-	UE_API inline static FOnNewTimeRulerStyle        OnNewTimeRulerStyle;
-	UE_API inline static FOnNewValueGridOverlayStyle OnNewValueGridOverlayStyle;
-	UE_API inline static FOnNewWaveformViewerStyle   OnNewWaveformViewerStyle;
+	UE_API static FOnNewTimeRulerStyle        OnNewTimeRulerStyle;
+	UE_API static FOnNewValueGridOverlayStyle OnNewValueGridOverlayStyle;
+	UE_API static FOnNewWaveformViewerStyle   OnNewWaveformViewerStyle;
 };
 
 #undef UE_API
diff --git a/Engine/Plugins/Runtime/AudioWidgets/Source/AudioWidgets/Public/AudioVectorscopePanelStyle.h b/Engine/Plugins/Runtime/AudioWidgets/Source/AudioWidgets/Public/AudioVectorscopePanelStyle.h
--- a/Engine/Plugins/Runtime/AudioWidgets/Source/AudioWidgets/Public/AudioVectorscopePanelStyle.h
+++ b/Engine/Plugins/Runtime/AudioWidgets/Source/AudioWidgets/Public/AudioVectorscopePanelStyle.h
@@ -21,7 +21,7 @@ struct FAudioVectorscopePanelStyle : public FSlateWidgetStyle
 {
 	GENERATED_USTRUCT_BODY()
 
-		FAudioVectorscopePanelStyle()
+	FAudioVectorscopePanelStyle()
 	{
 		ValueGridStyle = FSampledSequenceValueGridOverlayStyle()
 		.SetGridColor(FLinearColor(0.05f, 0.05f, 0.05f, 0.5f))
@@ -50,8 +50,8 @@ struct FAudioVectorscopePanelStyle : public FSlateWidgetStyle
 
 	inline static FName TypeName = FName("FAudioVectorscopePanelStyle");
 
-	UE_API inline static FOnNewValueGridOverlayStyle OnNewValueGridOverlayStyle;
-	UE_API inline static FOnNewVectorViewerStyle     OnNewVectorViewerStyle;
+	UE_API static FOnNewValueGridOverlayStyle OnNewValueGridOverlayStyle;
+	UE_API static FOnNewVectorViewerStyle     OnNewVectorViewerStyle;
 };
 
 #undef UE_API
