[AudioWidgets] Created AudioWidgetCore module, this module will contain widgets/logic that doesn't depend on Engine

diff --git a/Engine/Plugins/Runtime/AudioWidgets/AudioWidgets.uplugin b/Engine/Plugins/Runtime/AudioWidgets/AudioWidgets.uplugin
--- a/Engine/Plugins/Runtime/AudioWidgets/AudioWidgets.uplugin
+++ b/Engine/Plugins/Runtime/AudioWidgets/AudioWidgets.uplugin
@@ -14,7 +14,17 @@
 	"IsBetaVersion": false,
 	"IsExperimentalVersion": false,
 	"Installed": false,
+	"SupportedPrograms": [ "UnrealInsights" ],
 	"Modules": [
+		{
+			"Name": "AudioWidgetsCore",
+			"Type": "RuntimeAndProgram",
+			"LoadingPhase": "Default",
+			"ProgramAllowList":
+			[
+				"UnrealInsights"
+			]
+		},
 		{
 			"Name": "AudioWidgets",
 			"Type": "Runtime",
diff --git a/Engine/Plugins/Runtime/AudioWidgets/Source/AudioWidgets/AudioWidgets.Build.cs b/Engine/Plugins/Runtime/AudioWidgets/Source/AudioWidgets/AudioWidgets.Build.cs
--- a/Engine/Plugins/Runtime/AudioWidgets/Source/AudioWidgets/AudioWidgets.Build.cs
+++ b/Engine/Plugins/Runtime/AudioWidgets/Source/AudioWidgets/AudioWidgets.Build.cs
@@ -12,10 +12,11 @@ public class AudioWidgets : ModuleRules
 		(
 			new string[]
 			{
-				"Core",
-				"UMG",
 				"AudioAnalyzer",
-				"AudioSynesthesia"
+				"AudioSynesthesia",
+				"AudioWidgetsCore",
+				"Core",
+				"UMG"
 			}
 		);
 
diff --git a/Engine/Plugins/Runtime/AudioWidgets/Source/AudioWidgetsCore/AudioWidgetsCore.Build.cs b/Engine/Plugins/Runtime/AudioWidgets/Source/AudioWidgetsCore/AudioWidgetsCore.Build.cs
new file mode 100644
--- /dev/null
+++ b/Engine/Plugins/Runtime/AudioWidgets/Source/AudioWidgetsCore/AudioWidgetsCore.Build.cs
@@ -0,0 +1,31 @@
+// Copyright Epic Games, Inc. All Rights Reserved.
+
+using UnrealBuildTool;
+
+public class AudioWidgetsCore : ModuleRules
+{
+	public AudioWidgetsCore(ReadOnlyTargetRules Target) : base(Target)
+	{
+		PCHUsage = ModuleRules.PCHUsageMode.UseExplicitOrSharedPCHs;
+
+		PublicDependencyModuleNames.AddRange
+		(
+			new string[]
+			{
+				"Core",
+			}
+		);
+
+		PrivateDependencyModuleNames.AddRange
+		(
+			new string[]
+			{
+				"CoreUObject",
+				"Slate",
+				"SlateCore",
+			}
+		);
+
+		bDisableAutoRTFMInstrumentation = true;
+	}
+}
diff --git a/Engine/Plugins/Runtime/AudioWidgets/Source/AudioWidgetsCore/Private/AudioWidgetsCoreModule.cpp b/Engine/Plugins/Runtime/AudioWidgets/Source/AudioWidgetsCore/Private/AudioWidgetsCoreModule.cpp
new file mode 100644
--- /dev/null
+++ b/Engine/Plugins/Runtime/AudioWidgets/Source/AudioWidgetsCore/Private/AudioWidgetsCoreModule.cpp
@@ -0,0 +1,15 @@
+// Copyright Epic Games, Inc. All Rights Reserved.
+
+#include "AudioWidgetsCoreModule.h"
+
+#include "Modules/ModuleManager.h"
+
+void FAudioWidgetsCoreModule::StartupModule()
+{
+}
+
+void FAudioWidgetsCoreModule::ShutdownModule()
+{
+}
+
+IMPLEMENT_MODULE(FAudioWidgetsCoreModule, AudioWidgetsCore)
diff --git a/Engine/Plugins/Runtime/AudioWidgets/Source/AudioWidgetsCore/Public/AudioWidgetsCoreModule.h b/Engine/Plugins/Runtime/AudioWidgets/Source/AudioWidgetsCore/Public/AudioWidgetsCoreModule.h
new file mode 100644
--- /dev/null
+++ b/Engine/Plugins/Runtime/AudioWidgets/Source/AudioWidgetsCore/Public/AudioWidgetsCoreModule.h
@@ -0,0 +1,12 @@
+// Copyright Epic Games, Inc. All Rights Reserved.
+
+#pragma once
+
+#include "Modules/ModuleInterface.h"
+
+class FAudioWidgetsCoreModule : public IModuleInterface
+{
+public:
+	virtual void StartupModule() override;
+	virtual void ShutdownModule() override;
+};
