[AudioWidgets] Migration of FAudioMeterStyle, FAudioMeterDefaultColorStyle, FMeterChannelInfo and SAudioMeter logic to AudioWidgetsCore module in order to be able to use an audio meter widget without having to depend on Engine (so it can be used in Audio Insights standalone)

diff --git a/Engine/Plugins/Runtime/AudioWidgets/Source/AudioWidgets/Private/AudioMeter.cpp b/Engine/Plugins/Runtime/AudioWidgets/Source/AudioWidgets/Private/AudioMeter.cpp
--- a/Engine/Plugins/Runtime/AudioWidgets/Source/AudioWidgets/Private/AudioMeter.cpp
+++ b/Engine/Plugins/Runtime/AudioWidgets/Source/AudioWidgets/Private/AudioMeter.cpp
@@ -161,8 +161,6 @@ const FText UAudioMeter::GetPaletteCategory()
 }
 #endif
 
-const FName FAudioMeterDefaultColorStyle::TypeName(TEXT("FAudioMeterDefaultColorStyle"));
-
 const FAudioMeterDefaultColorStyle& FAudioMeterDefaultColorStyle::GetDefault()
 {
 	static FAudioMeterDefaultColorStyle Style;
diff --git a/Engine/Plugins/Runtime/AudioWidgets/Source/AudioWidgets/Private/AudioMeterStyle.cpp b/Engine/Plugins/Runtime/AudioWidgets/Source/AudioWidgets/Private/AudioMeterStyle.cpp
--- a/Engine/Plugins/Runtime/AudioWidgets/Source/AudioWidgets/Private/AudioMeterStyle.cpp
+++ b/Engine/Plugins/Runtime/AudioWidgets/Source/AudioWidgets/Private/AudioMeterStyle.cpp
@@ -1,41 +1,20 @@
 // Copyright Epic Games, Inc. All Rights Reserved.
 
 #include "AudioMeterStyle.h"
-#include "Styling/StyleDefaults.h"
 
 #include UE_INLINE_GENERATED_CPP_BY_NAME(AudioMeterStyle)
 
-FAudioMeterStyle::FAudioMeterStyle()
-	: MeterSize(FVector2D(250.0f, 25.0f))
-	, MeterPadding(FVector2D(10.0f, 5.0f))
-	, MeterValuePadding(3.0f)
-	, PeakValueWidth(2.0f)
-	, ValueRangeDb(FVector2D(-160, 10))
-	, bShowScale(true)
-	, bScaleSide(true)
-	, ScaleHashOffset(5.0f)
-	, ScaleHashWidth(1.0f)
-	, ScaleHashHeight(10.0f)
-	, DecibelsPerHash(10)
-	, Font(FStyleDefaults::GetFontInfo(5))
-{
-}
-
-void FAudioMeterStyle::GetResources(TArray< const FSlateBrush* >& OutBrushes) const
-{
-	OutBrushes.Add(&MeterValueImage);
-	OutBrushes.Add(&MeterBackgroundImage);
-	OutBrushes.Add(&MeterPeakImage);
-}
-
-const FName FAudioMeterStyle::TypeName(TEXT("FAudioMeterStyle"));
-
 const FAudioMeterStyle& FAudioMeterStyle::GetDefault()
 {
 	static FAudioMeterStyle Default;
 	return Default;
 }
 
+void FAudioMeterStyle::GetResources(TArray<const FSlateBrush*>& OutBrushes) const
+{
+	FAudioMeterWidgetStyle::GetResources(OutBrushes);
+}
+
 PRAGMA_DISABLE_DEPRECATION_WARNINGS
 FAudioMeterStyle& FAudioMeterStyle::SetFont(const FName& InFontName, uint16 InSize) { Font = FSlateFontInfo(InFontName, InSize); return *this; }
 FAudioMeterStyle& FAudioMeterStyle::SetFont(const FString& InFontName, uint16 InSize) { Font = FSlateFontInfo(*InFontName, InSize); return *this; }
@@ -45,4 +24,4 @@ FAudioMeterStyle& FAudioMeterStyle::SetFontName(const FName& InFontName) { Font
 FAudioMeterStyle& FAudioMeterStyle::SetFontName(const FString& InFontName) { Font = FSlateFontInfo(InFontName, Font.Size); return *this; }
 FAudioMeterStyle& FAudioMeterStyle::SetFontName(const WIDECHAR* InFontName) { Font = FSlateFontInfo(InFontName, Font.Size); return *this; }
 FAudioMeterStyle& FAudioMeterStyle::SetFontName(const ANSICHAR* InFontName) { Font = FSlateFontInfo(InFontName, Font.Size); return *this; }
-PRAGMA_ENABLE_DEPRECATION_WARNINGS
\ No newline at end of file
+PRAGMA_ENABLE_DEPRECATION_WARNINGS
diff --git a/Engine/Plugins/Runtime/AudioWidgets/Source/AudioWidgets/Private/SAudioMeter.cpp b/Engine/Plugins/Runtime/AudioWidgets/Source/AudioWidgets/Private/SAudioMeter.cpp
--- a/Engine/Plugins/Runtime/AudioWidgets/Source/AudioWidgets/Private/SAudioMeter.cpp
+++ b/Engine/Plugins/Runtime/AudioWidgets/Source/AudioWidgets/Private/SAudioMeter.cpp
@@ -1,407 +1,111 @@
 // Copyright Epic Games, Inc. All Rights Reserved.
 
 #include "SAudioMeter.h"
-#include "Framework/Application/SlateApplication.h"
-#include "Fonts/FontMeasure.h"
-#include "Rendering/SlateRenderer.h"
 
-SAudioMeter::SAudioMeter()
-{
-}
-
-void SAudioMeter::Construct(const SAudioMeter::FArguments& InArgs)
-{
-	check(InArgs._Style);
-
-	Orientation = InArgs._Orientation;
+#include "SAudioMeterWidget.h"
 
-	BackgroundColor = InArgs._BackgroundColor;
-	MeterBackgroundColor = InArgs._MeterBackgroundColor;
-	MeterValueColor = InArgs._MeterValueColor;
-	MeterPeakColor = InArgs._MeterPeakColor;
-	MeterScaleColor = InArgs._MeterScaleColor;
-	MeterScaleLabelColor = InArgs._MeterScaleLabelColor;
-	MeterClippingColor = InArgs._MeterClippingColor;
-
-	Style = InArgs._Style;
-
-	MeterChannelInfoAttribute = InArgs._MeterChannelInfo;
-}
-
-int32 SAudioMeter::OnPaint(const FPaintArgs& Args, const FGeometry& AllottedGeometry, const FSlateRect& MyCullingRect, FSlateWindowElementList& OutDrawElements, int32 LayerId, const FWidgetStyle& InWidgetStyle, bool bParentEnabled) const
+namespace SAudioMeterPrivate
 {
-	// we draw the slider like a horizontal slider regardless of the orientation, and apply a render transform to make it display correctly.
-	// However, the AllottedGeometry is computed as it will be rendered, so we have to use the "horizontal orientation" when doing drawing computations.
-	const float AllottedWidth = Orientation == Orient_Horizontal ? AllottedGeometry.GetLocalSize().X : AllottedGeometry.GetLocalSize().Y;
-	const float AllottedHeight = Orientation == Orient_Horizontal ? AllottedGeometry.GetLocalSize().Y : AllottedGeometry.GetLocalSize().X;
-
-	FGeometry MeterGeometry = AllottedGeometry;
-
-	// rotate the slider 90deg if it's vertical. The 0 side goes on the bottom, the 1 side on the top.
-	if (Orientation == Orient_Vertical)
-	{
-		// Do this by translating along -X by the width of the geometry, then rotating 90 degrees CCW (left-hand coords)
-		FSlateRenderTransform SlateRenderTransform = TransformCast<FSlateRenderTransform>(Concatenate(Inverse(FVector2D(AllottedWidth, 0)), FQuat2D(FMath::DegreesToRadians(-90.0f))));
-		// create a child geometry matching this one, but with the render transform.
-		MeterGeometry = AllottedGeometry.MakeChild(
-			FVector2D(AllottedWidth, AllottedHeight),
-			FSlateLayoutTransform(),
-			SlateRenderTransform, FVector2D::ZeroVector);
-	}
- 
- 	// We clamp to make sure that the slider cannot go out of the slider Length.
- 	int32 DecibelsPerHash = Style->DecibelsPerHash;
- 
- 	int32 MinValueDb = FMath::Min((int32)Style->ValueRangeDb.X, (int32)Style->ValueRangeDb.Y);
- 	int32 MaxValueDb = FMath::Max((int32)Style->ValueRangeDb.Y, (int32)Style->ValueRangeDb.Y);
- 
- 	// Snap the min/max values to the nearest hash 
- 	MaxValueDb -= MaxValueDb % DecibelsPerHash;
- 	MinValueDb -= MinValueDb % DecibelsPerHash;
- 
- 	float ValueDeltaDb = MaxValueDb - MinValueDb;
-
-	// Draw the background image
-	{
-		FVector2D BackgroundImageTopLeft = FVector2D(0.0f, 0.0f);
-		FVector2D BackgroundImageSize = FVector2D(AllottedWidth, AllottedHeight);
-
-		// Draw background image
-		FSlateDrawElement::MakeBox(
-			OutDrawElements,
-			LayerId,
-			MeterGeometry.ToPaintGeometry(BackgroundImageSize, FSlateLayoutTransform(BackgroundImageTopLeft)),
-			&Style->MeterBackgroundImage,
-			ESlateDrawEffect::None,
-			Style->BackgroundImage.GetTint(InWidgetStyle) * BackgroundColor.Get().GetColor(InWidgetStyle) * InWidgetStyle.GetColorAndOpacityTint()
-		);
-
-	}
- 	++LayerId;
-
-	float ScaleOffset = 0.0f;
-	if (Style->bShowScale && Style->bScaleSide)
+	const TAttribute<TArray<FAudioMeterChannelInfo>> ConvertMeterChannelInfoAttribute(const TAttribute<TArray<FMeterChannelInfo>>& InMeterChannelInfoAttribute)
 	{
-		 ScaleOffset = GetScaleHeight();
-	}
-
-	TArray<FMeterChannelInfo> ChannelInfos = MeterChannelInfoAttribute.Get();
-	int32 NumChannels = ChannelInfos.Num();
-
-	// Draw the meter backgrounds
-	for (int32 ChannelIndex = 0; ChannelIndex < NumChannels; ++ChannelIndex)
-	{		
-		FVector2D MeterBackgroundTopLeft = FVector2D(Style->MeterPadding.X, Style->MeterPadding.Y + ScaleOffset + ChannelIndex * (Style->MeterSize.Y + Style->MeterPadding.Y));
-		FVector2D MeterBackgroundSize = FVector2D(AllottedWidth - 2.0f * Style->MeterPadding.X, Style->MeterSize.Y);
-
- 		FSlateDrawElement::MakeBox(
-			OutDrawElements,
-			LayerId,
-			MeterGeometry.ToPaintGeometry(MeterBackgroundSize, FSlateLayoutTransform(MeterBackgroundTopLeft)),
-			&Style->MeterBackgroundImage,
-			ESlateDrawEffect::None,
-			Style->MeterBackgroundImage.GetTint(InWidgetStyle) * MeterBackgroundColor.Get().GetColor(InWidgetStyle) * InWidgetStyle.GetColorAndOpacityTint()
-		);
-	}
-
-	++LayerId;
-
-	// Draw the meter value
-	for (int32 ChannelIndex = 0; ChannelIndex < NumChannels; ++ChannelIndex)
-	{
-		float ChannelMeterValueDb = ChannelInfos[ChannelIndex].MeterValue;
-		float ChannelMeterValue = FMath::Clamp((ChannelMeterValueDb - MinValueDb) / ValueDeltaDb, 0.0f, 1.0f);
-
-		FVector2D MeterValueTopLeft = FVector2D(Style->MeterPadding.X, 
-												Style->MeterPadding.Y + ScaleOffset + Style->MeterValuePadding + ChannelIndex * (Style->MeterSize.Y + Style->MeterPadding.Y));
-		
-		FVector2D MeterValueSize = FVector2D(ChannelMeterValue * (AllottedWidth - 2.0f * Style->MeterPadding.X), 
-										     Style->MeterSize.Y - 2.0f * Style->MeterValuePadding);
- 
-		FSlateDrawElement::MakeBox(
-			OutDrawElements,
-			LayerId,
-			MeterGeometry.ToPaintGeometry(MeterValueSize, FSlateLayoutTransform(MeterValueTopLeft)),
-			&Style->MeterValueImage,
-			ESlateDrawEffect::None,
-			Style->MeterValueImage.GetTint(InWidgetStyle) * MeterValueColor.Get().GetColor(InWidgetStyle) * InWidgetStyle.GetColorAndOpacityTint()
-		);
-	}
-
-	++LayerId;
-
-	// Draw the meter peak value
-	for (int32 ChannelIndex = 0; ChannelIndex < NumChannels; ++ChannelIndex)
-	{	
-		float ChannelPeakValueDb = ChannelInfos[ChannelIndex].PeakValue;
-		float ChannelPeakValue = FMath::Clamp((ChannelPeakValueDb - MinValueDb) / ValueDeltaDb, 0.0f, 1.0f);
-
-		float PeakPixel = ChannelPeakValue * (AllottedWidth - 2.0f * Style->MeterPadding.X);
-
-		FVector2D MeterPeakValueTopLeft = FVector2D(Style->MeterPadding.X + PeakPixel - 0.5f * Style->PeakValueWidth, 
-													Style->MeterPadding.Y + ScaleOffset + Style->MeterValuePadding + ChannelIndex * (Style->MeterSize.Y + Style->MeterPadding.Y));
-		
-		FVector2D MeterPeakValueSize = FVector2D(Style->PeakValueWidth, 
-												 Style->MeterSize.Y - 2.0f * Style->MeterValuePadding);
-
-		FSlateDrawElement::MakeBox(
-			OutDrawElements,
-			LayerId,
-			MeterGeometry.ToPaintGeometry(MeterPeakValueSize, FSlateLayoutTransform(MeterPeakValueTopLeft)),
-			&Style->MeterPeakImage,
-			ESlateDrawEffect::None,
-			Style->MeterPeakImage.GetTint(InWidgetStyle) * MeterPeakColor.Get().GetColor(InWidgetStyle) * InWidgetStyle.GetColorAndOpacityTint()
-		);
-	}
-	
-	++LayerId;
-
-	// Draw the scale hash
-	if (Style->bShowScale)
-	{
-
-		float ScaleHashHalfWidth = 0.5f * Style->ScaleHashWidth;
-		FVector2D HashSize = FVector2D(Style->ScaleHashWidth, Style->ScaleHashHeight);
-
-		const TSharedRef<FSlateFontMeasure> FontMeasureService = FSlateApplication::Get().GetRenderer()->GetFontMeasureService();
-
-		// Measure the min value label size so we can right-justify if needed
-		FVector2D MinValueLabelSize = FontMeasureService->Measure(FString::FromInt(MinValueDb), Style->Font);
-
-		// Get the size of the negative sign to use to offset the label text in horizontal mode
-		FVector2D NegativeSignSize = FVector2D::ZeroVector;
-		if (Orientation == Orient_Horizontal)
-		{
-			NegativeSignSize = FontMeasureService->Measure(TEXT("-"), Style->Font);
-		}
-
-		int32 ValueDelta = MaxValueDb - MinValueDb;
-		int32 CurrentHashValue = MaxValueDb;
-		while (CurrentHashValue >= MinValueDb)
-		{
-			// Get the fractional value for hash mark
-			const float CurrentHashMeterValuePercent = FMath::Clamp(((float)CurrentHashValue - (float)MinValueDb) / ValueDelta, 0.0f, 1.0f);
-
-			// Draw hash
-			float HashPixelCenter = CurrentHashMeterValuePercent * (AllottedWidth - 2.0f * Style->MeterPadding.X);
-
-			FVector2D HashTopLeft;
-			HashTopLeft.X = Style->MeterPadding.X + HashPixelCenter - ScaleHashHalfWidth;
-			if (Style->bScaleSide)
-			{
-				HashTopLeft.Y = Style->MeterPadding.Y + ScaleOffset - Style->ScaleHashOffset - Style->ScaleHashHeight;
-			}
-			else
+		return TAttribute<TArray<FAudioMeterChannelInfo>>::Create([InMeterChannelInfoAttribute]()
 			{
-				HashTopLeft.Y = (Style->MeterSize.Y + Style->MeterPadding.Y ) * NumChannels + Style->ScaleHashOffset;
-			}
+				const TArray<FMeterChannelInfo>& MeterChannelInfoArray = InMeterChannelInfoAttribute.Get();
 
-			FSlateDrawElement::MakeBox(
-				OutDrawElements,
-				LayerId,
-				MeterGeometry.ToPaintGeometry(HashSize, FSlateLayoutTransform(HashTopLeft)),
-				&Style->MeterPeakImage,
-				ESlateDrawEffect::None,
-				MeterScaleColor.Get().GetColor(InWidgetStyle) * InWidgetStyle.GetColorAndOpacityTint()
-			);
-		
-			FString LabelString = FString::FromInt(CurrentHashValue);
+				TArray<FAudioMeterChannelInfo> AudioMeterChannelInfoArray;
+				AudioMeterChannelInfoArray.Reserve(MeterChannelInfoArray.Num());
 
-			bool bIsNegative = (CurrentHashValue < 0);
-
-			FVector2D LabelSize; 
-			if (Orientation == Orient_Horizontal && bIsNegative)
-			{
-				// We want to center the text on just the positive portion of the number
-				FString LabelStringPositive = FString::FromInt(FMath::Abs(CurrentHashValue));
-				LabelSize = FontMeasureService->Measure(LabelStringPositive, Style->Font);
-			}
-			else
-			{
-				LabelSize = FontMeasureService->Measure(LabelString, Style->Font);
-			}
-			
-			FText LabelText = FText::FromString(LabelString);
-
-			FGeometry TextGeometry = MeterGeometry;
-			if (Orientation == Orient_Vertical)
-			{
-				FVector2D LabelTopLeft;
-				
-				LabelTopLeft.X = Style->MeterPadding.X + HashPixelCenter + 0.5f * LabelSize.Y;
-
-				if (Style->bScaleSide)
-				{
-					LabelTopLeft.Y = Style->MeterPadding.Y - 2.0f + (MinValueLabelSize.X - LabelSize.X);
-				}
-				else
-				{
-					LabelTopLeft.Y = (Style->MeterSize.Y + Style->MeterPadding.Y) * NumChannels + Style->ScaleHashOffset + Style->ScaleHashHeight + 2.0f;
-				}
-				
-				// Undo the rotation for vertical right before we do the rendering of the scale value 
-				FSlateRenderTransform RotationTransform = FSlateRenderTransform(FQuat2D(FMath::DegreesToRadians(90.0f)));
-
-				FSlateLayoutTransform ChildLayoutTransform(1.0f, TransformPoint(1.0f, LabelTopLeft));
-
-				TextGeometry = MeterGeometry.MakeChild(
-					LabelSize,
-					ChildLayoutTransform,
-					RotationTransform,
-					FVector2D(0.0f, 0.0f)
-				);
-
-				FSlateDrawElement::MakeText(
-					OutDrawElements,
-					LayerId,
-					TextGeometry.ToOffsetPaintGeometry(FVector2D(0.0f, 0.0f)),
-					LabelText,
-					Style->Font,
-					ESlateDrawEffect::None,
-					MeterScaleLabelColor.Get().GetColor(InWidgetStyle)* InWidgetStyle.GetColorAndOpacityTint()
-				);
-			}
-			else
-			{
-				FVector2D LabelTopLeft;
-
-				// Center it based off the positive value width, then subtract from there the size of the negative sign
-				LabelTopLeft.X = Style->MeterPadding.X + HashPixelCenter - 0.5f * LabelSize.X;
-				if (bIsNegative)
-				{
-					LabelTopLeft.X -= NegativeSignSize.X;
-				}
-
-				if (Style->bScaleSide)
+				for (const FMeterChannelInfo& MeterChannelInfo : MeterChannelInfoArray)
 				{
-					LabelTopLeft.Y = Style->MeterPadding.Y + ScaleOffset - LabelSize.Y - Style->ScaleHashOffset - Style->ScaleHashHeight;
+					AudioMeterChannelInfoArray.Emplace(static_cast<const FAudioMeterChannelInfo&>(MeterChannelInfo));
 				}
-				else
-				{
-					LabelTopLeft.Y = (Style->MeterSize.Y + Style->MeterPadding.Y) * NumChannels + Style->ScaleHashOffset + Style->ScaleHashHeight + 2.0f;
-				}
-
-				// Draw text label
-				FSlateDrawElement::MakeText(
-					OutDrawElements,
-					LayerId,
-					TextGeometry.ToOffsetPaintGeometry(LabelTopLeft),
-					LabelText,
-					Style->Font,
-					ESlateDrawEffect::None,
-					MeterScaleLabelColor.Get().GetColor(InWidgetStyle) * InWidgetStyle.GetColorAndOpacityTint()
-				);
 
-			}
-
-			CurrentHashValue -= Style->DecibelsPerHash;
-		}
+				return AudioMeterChannelInfoArray;
+			});
 	}
-
-	return LayerId;
 }
 
-float SAudioMeter::GetScaleHeight() const
+void SAudioMeter::Construct(const SAudioMeter::FArguments& InArgs)
 {
-	const TSharedRef<FSlateFontMeasure> FontMeasureService = FSlateApplication::Get().GetRenderer()->GetFontMeasureService();
-	FVector2D LabelSize = FontMeasureService->Measure(FString::FromInt(-60), Style->Font);
-
-	float ScaleHeight = Style->ScaleHashHeight + Style->ScaleHashOffset;
-	if (Orientation == Orient_Vertical)
-	{
-		return ScaleHeight + LabelSize.X;
-	}
-	return ScaleHeight + LabelSize.Y;
+	SAssignNew(AudioMeterWidget, SAudioMeterWidget)
+	.Orientation(InArgs._Orientation)
+	.BackgroundColor(InArgs._BackgroundColor)
+	.MeterBackgroundColor(InArgs._MeterBackgroundColor)
+	.MeterValueColor(InArgs._MeterValueColor)
+	.MeterPeakColor(InArgs._MeterPeakColor)
+	.MeterScaleColor(InArgs._MeterScaleColor)
+	.MeterScaleLabelColor(InArgs._MeterScaleLabelColor)
+	.MeterClippingColor(InArgs._MeterClippingColor)
+	.Style(InArgs._Style)
+	.MeterChannelInfo(SAudioMeterPrivate::ConvertMeterChannelInfoAttribute(InArgs._MeterChannelInfo));
 }
 
-FVector2D SAudioMeter::ComputeDesiredSize(float) const
+int32 SAudioMeter::OnPaint(const FPaintArgs& Args, const FGeometry& AllottedGeometry, const FSlateRect& MyCullingRect, FSlateWindowElementList& OutDrawElements, int32 LayerId, const FWidgetStyle& InWidgetStyle, bool bParentEnabled) const
 {
-	static const FVector2D SAudioMeterDesiredSize(50.0f, 50.0f);
-
-	if (Style == nullptr)
-	{
-		return SAudioMeterDesiredSize;
-	}
-
-	TArray<FMeterChannelInfo> ChannelInfo = MeterChannelInfoAttribute.Get();
-	int32 NumChannels = FMath::Max(ChannelInfo.Num(), 1);
-
-	FVector2D Size = FVector2D(Style->MeterSize.X, (Style->MeterSize.Y + Style->MeterPadding.Y) * NumChannels);
-
-	// Add the end padding
-	Size.X += Style->MeterPadding.X;
-	Size.Y += Style->MeterPadding.Y;
-
-	// Add the width for the scale if it's been set to show
-	if (Style->bShowScale)
-	{
-		Size.Y += GetScaleHeight();
-	}
+	return AudioMeterWidget->OnPaint(Args, AllottedGeometry, MyCullingRect, OutDrawElements, LayerId, InWidgetStyle, bParentEnabled);
+}
 
-	// The thickness is going to be the thickness plus the scale image width
-	if (Orientation == Orient_Vertical)
-	{
-		return FVector2D(Size.Y, Size.X);
-	}
+float SAudioMeter::GetScaleHeight() const
+{
+	return AudioMeterWidget->GetScaleHeight();
+}
 
-	return Size; 
+FVector2D SAudioMeter::ComputeDesiredSize(float LayoutScaleMultiplier) const
+{
+	return AudioMeterWidget->ComputeDesiredSize(LayoutScaleMultiplier);
 }
 
 bool SAudioMeter::ComputeVolatility() const
 {
-	return true; 
+	return AudioMeterWidget->ComputeVolatility();
 }
 
 void SAudioMeter::SetMeterChannelInfo(const TAttribute<TArray<FMeterChannelInfo>>& InMeterChannelInfo)
 {
-	SetAttribute(MeterChannelInfoAttribute, InMeterChannelInfo, EInvalidateWidgetReason::Paint);
+    AudioMeterWidget->SetMeterChannelInfo(SAudioMeterPrivate::ConvertMeterChannelInfoAttribute(InMeterChannelInfo));
 }
 
 TArray<FMeterChannelInfo> SAudioMeter::GetMeterChannelInfo() const
 {
-	return MeterChannelInfoAttribute.Get();
+	return static_cast<TArray<FMeterChannelInfo>>(AudioMeterWidget->GetMeterChannelInfo());
 }
 
 void SAudioMeter::SetOrientation(EOrientation InOrientation)
 {
-	if (Orientation != InOrientation)
-	{
-		Orientation = InOrientation;
-		Invalidate(EInvalidateWidgetReason::Layout);
-	}
+	AudioMeterWidget->SetOrientation(InOrientation);
 }
 
 void SAudioMeter::SetBackgroundColor(FSlateColor InBackgroundColor)
 {
-	SetAttribute(BackgroundColor, TAttribute<FSlateColor>(InBackgroundColor), EInvalidateWidgetReason::Paint);
+	AudioMeterWidget->SetBackgroundColor(InBackgroundColor);
 }
 
 void SAudioMeter::SetMeterBackgroundColor(FSlateColor InMeterBackgroundColor)
 {
-	SetAttribute(MeterBackgroundColor, TAttribute<FSlateColor>(InMeterBackgroundColor), EInvalidateWidgetReason::Paint);
+	AudioMeterWidget->SetMeterBackgroundColor(InMeterBackgroundColor);
 }
 
 void SAudioMeter::SetMeterValueColor(FSlateColor InMeterValueColor)
 {
-	SetAttribute(MeterValueColor, TAttribute<FSlateColor>(InMeterValueColor), EInvalidateWidgetReason::Paint);
+	AudioMeterWidget->SetMeterValueColor(InMeterValueColor);
 }
 
 void SAudioMeter::SetMeterPeakColor(FSlateColor InMeterPeakColor)
 {
-	SetAttribute(MeterPeakColor, TAttribute<FSlateColor>(InMeterPeakColor), EInvalidateWidgetReason::Paint);
+	AudioMeterWidget->SetMeterPeakColor(InMeterPeakColor);
 }
 
 void SAudioMeter::SetMeterClippingColor(FSlateColor InMeterClippingColor)
 {
-	SetAttribute(MeterClippingColor, TAttribute<FSlateColor>(InMeterClippingColor), EInvalidateWidgetReason::Paint);
+	AudioMeterWidget->SetMeterClippingColor(InMeterClippingColor);
 }
 
 void SAudioMeter::SetMeterScaleColor(FSlateColor InMeterScaleColor)
 {
-	SetAttribute(MeterScaleColor, TAttribute<FSlateColor>(InMeterScaleColor), EInvalidateWidgetReason::Paint);
+	AudioMeterWidget->SetMeterScaleColor(InMeterScaleColor);
 }
 
 void SAudioMeter::SetMeterScaleLabelColor(FSlateColor InMeterScaleLabelColor)
 {
-	SetAttribute(MeterScaleLabelColor, TAttribute<FSlateColor>(InMeterScaleLabelColor), EInvalidateWidgetReason::Paint);
+	AudioMeterWidget->SetMeterScaleLabelColor(InMeterScaleLabelColor);
 }
diff --git a/Engine/Plugins/Runtime/AudioWidgets/Source/AudioWidgets/Public/AudioMeter.h b/Engine/Plugins/Runtime/AudioWidgets/Source/AudioWidgets/Public/AudioMeter.h
--- a/Engine/Plugins/Runtime/AudioWidgets/Source/AudioWidgets/Public/AudioMeter.h
+++ b/Engine/Plugins/Runtime/AudioWidgets/Source/AudioWidgets/Public/AudioMeter.h
@@ -6,6 +6,7 @@
 #include "AudioDefines.h"
 #include "AudioMeterStyle.h"
 #include "AudioMeterTypes.h"
+#include "AudioMeterWidgetStyle.h"
 #include "Components/Widget.h"
 #include "Delegates/IDelegateInstance.h"
 #include "Engine/World.h"
@@ -33,33 +34,15 @@ struct FAudioMaterialMeterStyle;
 
 
 USTRUCT(BlueprintType)
-struct FAudioMeterDefaultColorStyle : public FSlateWidgetStyle
+struct FAudioMeterDefaultColorStyle : public FAudioMeterDefaultColorWidgetStyle
 {
-	GENERATED_USTRUCT_BODY()
+	GENERATED_BODY()
 
-public:
-	static UE_API const FName TypeName;
-	virtual const FName GetTypeName() const override { return TypeName; };
-
-	static UE_API const FAudioMeterDefaultColorStyle& GetDefault();
-
-	UPROPERTY(EditAnywhere, BlueprintReadWrite, Category = "Style")
-	FLinearColor MeterBackgroundColor = FLinearColor(0.031f, 0.031f, 0.031f, 1.0f);
+	UE_API static const FAudioMeterDefaultColorStyle& GetDefault();
 
-	UPROPERTY(EditAnywhere, BlueprintReadWrite, Category = "Style")
-	FLinearColor MeterValueColor = FLinearColor(0.025719f, 0.208333f, 0.069907f, 1.0f);
-
-	UPROPERTY(EditAnywhere, BlueprintReadWrite, Category = "Style")
-	FLinearColor MeterPeakColor = FLinearColor(0.24349f, 0.708333f, 0.357002f, 1.0f);
-
-	UPROPERTY(EditAnywhere, BlueprintReadWrite, Category = "Style")
-	FLinearColor MeterClippingColor = FLinearColor(1.0f, 0.0f, 0.112334f, 1.0f);
-
-	UPROPERTY(EditAnywhere, BlueprintReadWrite, Category = "Style")
-	FLinearColor MeterScaleColor = FLinearColor(0.017642f, 0.017642f, 0.017642f, 1.0f);
+	virtual const FName GetTypeName() const override { return TypeName; };
 
-	UPROPERTY(EditAnywhere, BlueprintReadWrite, Category = "Style")
-	FLinearColor MeterScaleLabelColor = FLinearColor(0.442708f, 0.442708f, 0.442708f, 1.0f);
+	inline static const FName TypeName = "FAudioMeterDefaultColorStyle";
 };
 
 /**
diff --git a/Engine/Plugins/Runtime/AudioWidgets/Source/AudioWidgets/Public/AudioMeterStyle.h b/Engine/Plugins/Runtime/AudioWidgets/Source/AudioWidgets/Public/AudioMeterStyle.h
--- a/Engine/Plugins/Runtime/AudioWidgets/Source/AudioWidgets/Public/AudioMeterStyle.h
+++ b/Engine/Plugins/Runtime/AudioWidgets/Source/AudioWidgets/Public/AudioMeterStyle.h
@@ -2,152 +2,44 @@
 
 #pragma once
 
-#include "Fonts/SlateFontInfo.h"
-#include "Styling/SlateWidgetStyle.h"
-#include "Styling/SlateBrush.h"
+#include "AudioMeterWidgetStyle.h"
+
 #include "AudioMeterStyle.generated.h"
 
 #define UE_API AUDIOWIDGETS_API
 
-struct FCompositeFont;
-
 /**
  * Represents the appearance of an SAudioMeter
  */
 USTRUCT(BlueprintType)
-struct FAudioMeterStyle : public FSlateWidgetStyle
+struct UE_API FAudioMeterStyle : public FAudioMeterWidgetStyle
 {
-	GENERATED_USTRUCT_BODY()
-
-	UE_API FAudioMeterStyle();
+	GENERATED_BODY()
 
-	virtual ~FAudioMeterStyle() {}
+	static const FAudioMeterStyle& GetDefault();
 
-	UE_API virtual void GetResources( TArray< const FSlateBrush* >& OutBrushes ) const override;
+	virtual void GetResources(TArray<const FSlateBrush*>& OutBrushes) const override;
 
-	static UE_API const FName TypeName;
 	virtual const FName GetTypeName() const override { return TypeName; };
 
-	static UE_API const FAudioMeterStyle& GetDefault();
-
-	/** Image to use to represent the meter value. */
-	UPROPERTY(EditAnywhere, BlueprintReadWrite, Category=Appearance)
-	FSlateBrush MeterValueImage;
-	FAudioMeterStyle& SetMeterValueImage(const FSlateBrush& InMeterValueImage){ MeterValueImage = InMeterValueImage; return *this; }
-
-	/** Image to use to represent the background. */
-	UPROPERTY(EditAnywhere, BlueprintReadWrite, Category = Appearance)
-	FSlateBrush BackgroundImage;
-	FAudioMeterStyle& SetBackgroundImage(const FSlateBrush& InBackgroundImage) { BackgroundImage = InBackgroundImage; return *this; }
-
-	/** Image to use to represent the meter background. */
-	UPROPERTY(EditAnywhere, BlueprintReadWrite, Category = Appearance)
-	FSlateBrush MeterBackgroundImage;
-	FAudioMeterStyle& SetMeterBackgroundImage(const FSlateBrush& InMeterBackgroundImage) { MeterBackgroundImage = InMeterBackgroundImage; return *this; }
-
-	/** Image to use to draw behind the meter value. */
-	UPROPERTY(EditAnywhere, BlueprintReadWrite, Category = Appearance)
-	FSlateBrush MeterValueBackgroundImage;
-	FAudioMeterStyle& SetMeterValueBackgroundImage(const FSlateBrush& InMeterValueBackgroundImage) { MeterValueBackgroundImage = InMeterValueBackgroundImage; return *this; }
-
-	/** Image to use to represent the meter peak. */
-	UPROPERTY(EditAnywhere, BlueprintReadWrite, Category = Appearance)
-	FSlateBrush MeterPeakImage;
-	FAudioMeterStyle& SetMeterPeakImage(const FSlateBrush& InMeterPeakImage) { MeterPeakImage = InMeterPeakImage; return *this; }
-
-	// How thick to draw the meter
-	UPROPERTY(EditAnywhere, BlueprintReadWrite, Category=Appearance)
-	FVector2D MeterSize;
-	FAudioMeterStyle& SetMeterSize(const FVector2D& InMeterSize) { MeterSize = InMeterSize; return *this; }
-
-	// How much padding to add around the meter
-	UPROPERTY(EditAnywhere, BlueprintReadWrite, Category = Appearance)
-	FVector2D MeterPadding;
-	FAudioMeterStyle& SetMeterPadding(const FVector2D& InMeterPadding) { MeterPadding = InMeterPadding; return *this; }
-
-	// How much padding to add around the meter value
-	UPROPERTY(EditAnywhere, BlueprintReadWrite, Category = Appearance)
-	float MeterValuePadding;
-	FAudioMeterStyle& SetMeterValuePadding(float InMeterValuePadding) { MeterValuePadding = InMeterValuePadding; return *this; }
-
-	// How wide to draw the peak value indicator
-	UPROPERTY(EditAnywhere, BlueprintReadWrite, Category = Appearance)
-	float PeakValueWidth;
-	FAudioMeterStyle& SetPeakValueWidth(float InPeakValueWidth) { PeakValueWidth = InPeakValueWidth; return *this; }
-
-	// The minimum and maximum value to display in dB (values are clamped in this range)
-	UPROPERTY(EditAnywhere, BlueprintReadWrite, Category = Appearance)
-	FVector2D ValueRangeDb;
-	FAudioMeterStyle& SetValueRangeDb(const FVector2D& InValueRangeDb) { ValueRangeDb = InValueRangeDb; return *this; }
-
-	// Whether or not to show the decibel scale alongside the meter
-	UPROPERTY(EditAnywhere, BlueprintReadWrite, Category = Appearance)
-	bool bShowScale;
-	FAudioMeterStyle& SetShowScale(bool bInShowScale) { bShowScale = bInShowScale; return *this; }
-
-	// Which side to show the scale. If vertical, true means left side, false means right side. If horizontal, true means above, false means below.
-	UPROPERTY(EditAnywhere, BlueprintReadWrite, Category = Appearance)
-	bool bScaleSide;
-	FAudioMeterStyle& SetScaleSide(bool bInScaleSide) { bScaleSide = bInScaleSide; return *this; }
-
-	// Offset for the hashes
-	UPROPERTY(EditAnywhere, BlueprintReadWrite, Category = Appearance)
-	float ScaleHashOffset;
-	FAudioMeterStyle& SetScaleHashOffset(float InScaleHashOffset) { ScaleHashOffset = InScaleHashOffset; return *this; }
-
-	// The width of each hash mark
-	UPROPERTY(EditAnywhere, BlueprintReadWrite, Category = Appearance)
-	float ScaleHashWidth;
-	FAudioMeterStyle& SetScaleHashWidth(float InScaleHashWidth) { ScaleHashWidth = InScaleHashWidth; return *this; }
-
-	// The height of each hash mark
-	UPROPERTY(EditAnywhere, BlueprintReadWrite, Category = Appearance)
-	float ScaleHashHeight;
-	FAudioMeterStyle& SetScaleHashHeight(float InScaleHashHeight) { ScaleHashHeight = InScaleHashHeight; return *this; }
-
-	// How wide to draw the decibel scale, if it's enabled
-	UPROPERTY(EditAnywhere, BlueprintReadWrite, Category = Appearance, meta = (UIMin = "3", ClampMin="3", UIMax = "10"))
-	int32 DecibelsPerHash;
-	FAudioMeterStyle& SetDecibelsPerHash(float InDecibelsPerHash) { DecibelsPerHash = InDecibelsPerHash; return *this; }
-
-	/** Font family and size to be used when displaying the meter scale. */
-	UPROPERTY(EditAnywhere, BlueprintReadWrite, Category = Appearance)
-	FSlateFontInfo Font;
-	FAudioMeterStyle& SetFont(const FSlateFontInfo& InFont) { Font = InFont; return *this; }
-	FAudioMeterStyle& SetFont(TSharedPtr<const FCompositeFont> InCompositeFont, const int32 InSize, const FName& InTypefaceFontName = NAME_None) { Font = FSlateFontInfo(InCompositeFont, InSize, InTypefaceFontName); return *this; }
-	FAudioMeterStyle& SetFont(const UObject* InFontObject, const int32 InSize, const FName& InTypefaceFontName = NAME_None) { Font = FSlateFontInfo(InFontObject, InSize, InTypefaceFontName); return *this; }
+	inline static const FName TypeName = "FAudioMeterStyle";
 
 	UE_DEPRECATED(5.6, "Use SetFont using FSlateFontInfo instead. FSlateFontInfo's constructors using a FontName are deprecated.")
-	UE_API FAudioMeterStyle& SetFont(const FName& InFontName, uint16 InSize);
+	FAudioMeterStyle& SetFont(const FName& InFontName, uint16 InSize);
 	UE_DEPRECATED(5.6, "Use SetFont using FSlateFontInfo instead. FSlateFontInfo's constructors using a FontName are deprecated.")
-	UE_API FAudioMeterStyle& SetFont(const FString& InFontName, uint16 InSize);
+	FAudioMeterStyle& SetFont(const FString& InFontName, uint16 InSize);
 	UE_DEPRECATED(5.6, "Use SetFont using FSlateFontInfo instead. FSlateFontInfo's constructors using a FontName are deprecated.")
-	UE_API FAudioMeterStyle& SetFont(const WIDECHAR* InFontName, uint16 InSize);
+	FAudioMeterStyle& SetFont(const WIDECHAR* InFontName, uint16 InSize);
 	UE_DEPRECATED(5.6, "Use SetFont using FSlateFontInfo instead. FSlateFontInfo's constructors using a FontName are deprecated.")
-	UE_API FAudioMeterStyle& SetFont(const ANSICHAR* InFontName, uint16 InSize);
+	FAudioMeterStyle& SetFont(const ANSICHAR* InFontName, uint16 InSize);
 	UE_DEPRECATED(5.6, "Use SetFont using FSlateFontInfo instead. FSlateFontInfo's constructors using a FontName are deprecated.")
-	UE_API FAudioMeterStyle& SetFontName(const FName& InFontName);
+	FAudioMeterStyle& SetFontName(const FName& InFontName);
 	UE_DEPRECATED(5.6, "Use SetFont using FSlateFontInfo instead. FSlateFontInfo's constructors using a FontName are deprecated.")
-	UE_API FAudioMeterStyle& SetFontName(const FString& InFontName);
+	FAudioMeterStyle& SetFontName(const FString& InFontName);
 	UE_DEPRECATED(5.6, "Use SetFont using FSlateFontInfo instead. FSlateFontInfo's constructors using a FontName are deprecated.")
-	UE_API FAudioMeterStyle& SetFontName(const WIDECHAR* InFontName);
+	FAudioMeterStyle& SetFontName(const WIDECHAR* InFontName);
 	UE_DEPRECATED(5.6, "Use SetFont using FSlateFontInfo instead. FSlateFontInfo's constructors using a FontName are deprecated.")
-	UE_API FAudioMeterStyle& SetFontName(const ANSICHAR* InFontName);
-
-	FAudioMeterStyle& SetFontSize(uint16 InSize) { Font.Size = InSize; return *this; }
-	FAudioMeterStyle& SetTypefaceFontName(const FName& InTypefaceFontName) { Font.TypefaceFontName = InTypefaceFontName; return *this; }
-
-	/**
-	 * Unlinks all colors in this style.
-	 * @see FSlateColor::Unlink
-	 */
-	void UnlinkColors()
-	{
-		MeterValueImage.UnlinkColors();
-		MeterBackgroundImage.UnlinkColors();
-//		MeterScaleImage.UnlinkColors();
-		MeterPeakImage.UnlinkColors();
-	}
+	FAudioMeterStyle& SetFontName(const ANSICHAR* InFontName);
 };
 
 #undef UE_API
diff --git a/Engine/Plugins/Runtime/AudioWidgets/Source/AudioWidgets/Public/AudioMeterTypes.h b/Engine/Plugins/Runtime/AudioWidgets/Source/AudioWidgets/Public/AudioMeterTypes.h
--- a/Engine/Plugins/Runtime/AudioWidgets/Source/AudioWidgets/Public/AudioMeterTypes.h
+++ b/Engine/Plugins/Runtime/AudioWidgets/Source/AudioWidgets/Public/AudioMeterTypes.h
@@ -2,26 +2,12 @@
 
 #pragma once
 
+#include "AudioMeterChannelInfo.h"
+
 #include "AudioMeterTypes.generated.h"
 
 USTRUCT(BlueprintType)
-struct FMeterChannelInfo
+struct FMeterChannelInfo : public FAudioMeterChannelInfo
 {
 	GENERATED_BODY()
-
-	UPROPERTY(EditAnywhere, BlueprintReadWrite, Category = AudioMeter)
-	float MeterValue = 0.0f;
-
-	UPROPERTY(EditAnywhere, BlueprintReadWrite, Category = AudioMeter)
-	float PeakValue = 0.0f;
-
-	UPROPERTY(EditAnywhere, BlueprintReadWrite, Category = AudioMeter)
-	float ClippingValue = 0.0f;
 };
-
-inline bool operator==(const FMeterChannelInfo& lhs, const FMeterChannelInfo& rhs)
-{
-	return FMath::IsNearlyEqual(lhs.MeterValue, rhs.MeterValue) &&
-		FMath::IsNearlyEqual(lhs.PeakValue, rhs.PeakValue) &&
-		FMath::IsNearlyEqual(lhs.ClippingValue, rhs.ClippingValue);
-}
diff --git a/Engine/Plugins/Runtime/AudioWidgets/Source/AudioWidgets/Public/SAudioMeter.h b/Engine/Plugins/Runtime/AudioWidgets/Source/AudioWidgets/Public/SAudioMeter.h
--- a/Engine/Plugins/Runtime/AudioWidgets/Source/AudioWidgets/Public/SAudioMeter.h
+++ b/Engine/Plugins/Runtime/AudioWidgets/Source/AudioWidgets/Public/SAudioMeter.h
@@ -2,13 +2,13 @@
 
 #pragma once
 
+#include "AudioMeterStyle.h"
+#include "AudioMeterTypes.h"
 #include "AudioWidgetsStyle.h"
 #include "Styling/ISlateStyle.h"
 #include "Styling/SlateWidgetStyleAsset.h"
 #include "Widgets/DeclarativeSyntaxSupport.h"
 #include "Widgets/SLeafWidget.h"
-#include "AudioMeterStyle.h"
-#include "AudioMeterTypes.h"
 
 #define UE_API AUDIOWIDGETS_API
 
@@ -35,7 +35,7 @@ public:
 };
 
 /**
- * A Slate slider control is a linear scale and draggable handle.
+ * A widget that displays audio channel's levels.
  */
 class SAudioMeter
 	: public SAudioMeterBase
@@ -57,7 +57,7 @@ public:
 	/** Whether the slidable area should be indented to fit the handle. */
 	SLATE_ATTRIBUTE(bool, IndentHandle)
 
-		/** The slider's orientation. */
+		/** The audio meter's orientation. */
 		SLATE_ARGUMENT(EOrientation, Orientation)
 
 		/** The color to draw the background in. */
@@ -89,8 +89,6 @@ public:
 
 		SLATE_END_ARGS()
 
-		UE_API SAudioMeter();
-
 	/**
 	 * Construct the widget.
 	 *
@@ -117,30 +115,15 @@ public:
 	// SWidget overrides
 
 	UE_API virtual int32 OnPaint(const FPaintArgs& Args, const FGeometry& AllottedGeometry, const FSlateRect& MyCullingRect, FSlateWindowElementList& OutDrawElements, int32 LayerId, const FWidgetStyle& InWidgetStyle, bool bParentEnabled) const override;
-	UE_API virtual FVector2D ComputeDesiredSize(float) const override;
+	UE_API virtual FVector2D ComputeDesiredSize(float LayoutScaleMultiplier) const override;
 	UE_API virtual bool ComputeVolatility() const override;
 
 protected:
 	
 	// Returns the scale height based off font size and hash height
 	UE_API float GetScaleHeight() const;
-	
-	// Holds the style passed to the widget upon construction.
-	const FAudioMeterStyle* Style;
-
-	// Holds the slider's orientation.
-	EOrientation Orientation;
-
-	// Various colors
-	TAttribute<FSlateColor> BackgroundColor;
-	TAttribute<FSlateColor> MeterBackgroundColor;
-	TAttribute<FSlateColor> MeterValueColor;
-	TAttribute<FSlateColor> MeterPeakColor;
-	TAttribute<FSlateColor> MeterClippingColor;
-	TAttribute<FSlateColor> MeterScaleColor;
-	TAttribute<FSlateColor> MeterScaleLabelColor;
-
-	TAttribute<TArray<FMeterChannelInfo>> MeterChannelInfoAttribute;
+
+	TSharedPtr<class SAudioMeterWidget> AudioMeterWidget;
 };
 
 #undef UE_API
diff --git a/Engine/Plugins/Runtime/AudioWidgets/Source/AudioWidgetsCore/Private/AudioMeterWidgetStyle.cpp b/Engine/Plugins/Runtime/AudioWidgets/Source/AudioWidgetsCore/Private/AudioMeterWidgetStyle.cpp
new file mode 100644
--- /dev/null
+++ b/Engine/Plugins/Runtime/AudioWidgets/Source/AudioWidgetsCore/Private/AudioMeterWidgetStyle.cpp
@@ -0,0 +1,43 @@
+// Copyright Epic Games, Inc. All Rights Reserved.
+
+#include "AudioMeterWidgetStyle.h"
+
+#include UE_INLINE_GENERATED_CPP_BY_NAME(AudioMeterWidgetStyle)
+
+const FAudioMeterWidgetStyle& FAudioMeterWidgetStyle::GetDefault()
+{
+	static FAudioMeterWidgetStyle Default;
+	return Default;
+}
+
+void FAudioMeterWidgetStyle::GetResources(TArray<const FSlateBrush*>& OutBrushes) const
+{
+	OutBrushes.Add(&MeterValueImage);
+	OutBrushes.Add(&MeterBackgroundImage);
+	OutBrushes.Add(&MeterPeakImage);
+}
+
+void FAudioMeterWidgetStyle::UnlinkColors()
+{
+	MeterValueImage.UnlinkColors();
+	MeterBackgroundImage.UnlinkColors();
+	MeterPeakImage.UnlinkColors();
+}
+
+PRAGMA_DISABLE_DEPRECATION_WARNINGS
+FAudioMeterWidgetStyle& FAudioMeterWidgetStyle::SetFont(const FName& InFontName, uint16 InSize) { Font = FSlateFontInfo(InFontName, InSize); return *this; }
+FAudioMeterWidgetStyle& FAudioMeterWidgetStyle::SetFont(const FString& InFontName, uint16 InSize) { Font = FSlateFontInfo(*InFontName, InSize); return *this; }
+FAudioMeterWidgetStyle& FAudioMeterWidgetStyle::SetFont(const WIDECHAR* InFontName, uint16 InSize) { Font = FSlateFontInfo(InFontName, InSize); return *this; }
+FAudioMeterWidgetStyle& FAudioMeterWidgetStyle::SetFont(const ANSICHAR* InFontName, uint16 InSize) { Font = FSlateFontInfo(InFontName, InSize); return *this; }
+FAudioMeterWidgetStyle& FAudioMeterWidgetStyle::SetFontName(const FName& InFontName) { Font = FSlateFontInfo(InFontName, Font.Size); return *this; }
+FAudioMeterWidgetStyle& FAudioMeterWidgetStyle::SetFontName(const FString& InFontName) { Font = FSlateFontInfo(InFontName, Font.Size); return *this; }
+FAudioMeterWidgetStyle& FAudioMeterWidgetStyle::SetFontName(const WIDECHAR* InFontName) { Font = FSlateFontInfo(InFontName, Font.Size); return *this; }
+FAudioMeterWidgetStyle& FAudioMeterWidgetStyle::SetFontName(const ANSICHAR* InFontName) { Font = FSlateFontInfo(InFontName, Font.Size); return *this; }
+PRAGMA_ENABLE_DEPRECATION_WARNINGS
+
+
+const FAudioMeterDefaultColorWidgetStyle& FAudioMeterDefaultColorWidgetStyle::GetDefault()
+{
+	static const FAudioMeterDefaultColorWidgetStyle Default;
+	return Default;
+}
diff --git a/Engine/Plugins/Runtime/AudioWidgets/Source/AudioWidgetsCore/Private/SAudioMeterWidget.cpp b/Engine/Plugins/Runtime/AudioWidgets/Source/AudioWidgetsCore/Private/SAudioMeterWidget.cpp
new file mode 100644
--- /dev/null
+++ b/Engine/Plugins/Runtime/AudioWidgets/Source/AudioWidgetsCore/Private/SAudioMeterWidget.cpp
@@ -0,0 +1,382 @@
+// Copyright Epic Games, Inc. All Rights Reserved.
+
+#include "SAudioMeterWidget.h"
+
+#include "Framework/Application/SlateApplication.h"
+#include "Fonts/FontMeasure.h"
+#include "Rendering/SlateRenderer.h"
+
+void SAudioMeterWidget::Construct(const SAudioMeterWidget::FArguments& InArgs)
+{
+	check(InArgs._Style);
+
+	Orientation = InArgs._Orientation;
+
+	BackgroundColor = InArgs._BackgroundColor;
+	MeterBackgroundColor = InArgs._MeterBackgroundColor;
+	MeterValueColor = InArgs._MeterValueColor;
+	MeterPeakColor = InArgs._MeterPeakColor;
+	MeterScaleColor = InArgs._MeterScaleColor;
+	MeterScaleLabelColor = InArgs._MeterScaleLabelColor;
+	MeterClippingColor = InArgs._MeterClippingColor;
+
+	Style = InArgs._Style;
+
+	MeterChannelInfoAttribute = InArgs._MeterChannelInfo;
+}
+
+int32 SAudioMeterWidget::OnPaint(const FPaintArgs& Args, const FGeometry& AllottedGeometry, const FSlateRect& MyCullingRect, FSlateWindowElementList& OutDrawElements, int32 LayerId, const FWidgetStyle& InWidgetStyle, bool bParentEnabled) const
+{
+	// we draw the meter like a horizontal meter regardless of the orientation, and apply a render transform to make it display correctly.
+	// However, the AllottedGeometry is computed as it will be rendered, so we have to use the "horizontal orientation" when doing drawing computations.
+	const float AllottedWidth  = Orientation == Orient_Horizontal ? AllottedGeometry.GetLocalSize().X : AllottedGeometry.GetLocalSize().Y;
+	const float AllottedHeight = Orientation == Orient_Horizontal ? AllottedGeometry.GetLocalSize().Y : AllottedGeometry.GetLocalSize().X;
+
+	FGeometry MeterGeometry = AllottedGeometry;
+
+	// rotate the meter 90deg if it's vertical. The 0 side goes on the bottom, the 1 side on the top.
+	if (Orientation == Orient_Vertical)
+	{
+		// Do this by translating along -X by the width of the geometry, then rotating 90 degrees CCW (left-hand coords)
+		const FSlateRenderTransform SlateRenderTransform = TransformCast<FSlateRenderTransform>(Concatenate(Inverse(FVector2D(AllottedWidth, 0)), FQuat2D(FMath::DegreesToRadians(-90.0f))));
+
+		// create a child geometry matching this one, but with the render transform.
+		MeterGeometry = AllottedGeometry.MakeChild(FVector2D(AllottedWidth, AllottedHeight), FSlateLayoutTransform(), SlateRenderTransform, FVector2D::ZeroVector);
+	}
+ 
+ 	// We clamp to make sure that the meter cannot go out of the meter Length.
+ 	const int32 DecibelsPerHash = Style->DecibelsPerHash;
+ 
+ 	int32 MinValueDb = FMath::Min(static_cast<int32>(Style->ValueRangeDb.X), static_cast<int32>(Style->ValueRangeDb.Y));
+ 	int32 MaxValueDb = FMath::Max(static_cast<int32>(Style->ValueRangeDb.Y), static_cast<int32>(Style->ValueRangeDb.Y));
+ 
+ 	// Snap the min/max values to the nearest hash 
+ 	MaxValueDb -= MaxValueDb % DecibelsPerHash;
+ 	MinValueDb -= MinValueDb % DecibelsPerHash;
+ 
+ 	const float ValueDeltaDb = MaxValueDb - MinValueDb;
+
+	// Draw the background image
+	{
+		const FVector2D BackgroundImageTopLeft(0.0f, 0.0f);
+		const FVector2D BackgroundImageSize(AllottedWidth, AllottedHeight);
+
+		// Draw background image
+		FSlateDrawElement::MakeBox(
+			OutDrawElements,
+			LayerId,
+			MeterGeometry.ToPaintGeometry(BackgroundImageSize, FSlateLayoutTransform(BackgroundImageTopLeft)),
+			&Style->MeterBackgroundImage,
+			ESlateDrawEffect::None,
+			Style->BackgroundImage.GetTint(InWidgetStyle) * BackgroundColor.Get().GetColor(InWidgetStyle) * InWidgetStyle.GetColorAndOpacityTint()
+		);
+	}
+ 	++LayerId;
+
+	float ScaleOffset = 0.0f;
+	if (Style->bShowScale && Style->bScaleSide)
+	{
+		 ScaleOffset = GetScaleHeight();
+	}
+
+	const TArray<FAudioMeterChannelInfo> ChannelInfos = MeterChannelInfoAttribute.Get();
+	const int32 NumChannels = ChannelInfos.Num();
+
+	// Draw the meter backgrounds
+	for (int32 ChannelIndex = 0; ChannelIndex < NumChannels; ++ChannelIndex)
+	{		
+		const FVector2D MeterBackgroundTopLeft(Style->MeterPadding.X, Style->MeterPadding.Y + ScaleOffset + ChannelIndex * (Style->MeterSize.Y + Style->MeterPadding.Y));
+		const FVector2D MeterBackgroundSize(AllottedWidth - 2.0f * Style->MeterPadding.X, Style->MeterSize.Y);
+
+ 		FSlateDrawElement::MakeBox(
+			OutDrawElements,
+			LayerId,
+			MeterGeometry.ToPaintGeometry(MeterBackgroundSize, FSlateLayoutTransform(MeterBackgroundTopLeft)),
+			&Style->MeterBackgroundImage,
+			ESlateDrawEffect::None,
+			Style->MeterBackgroundImage.GetTint(InWidgetStyle) * MeterBackgroundColor.Get().GetColor(InWidgetStyle) * InWidgetStyle.GetColorAndOpacityTint()
+		);
+	}
+	++LayerId;
+
+	// Draw the meter value
+	for (int32 ChannelIndex = 0; ChannelIndex < NumChannels; ++ChannelIndex)
+	{
+		const float ChannelMeterValueDb = ChannelInfos[ChannelIndex].MeterValue;
+		const float ChannelMeterValue = FMath::Clamp((ChannelMeterValueDb - MinValueDb) / ValueDeltaDb, 0.0f, 1.0f);
+
+		const FVector2D MeterValueTopLeft(Style->MeterPadding.X, 
+									Style->MeterPadding.Y + ScaleOffset + Style->MeterValuePadding + ChannelIndex * (Style->MeterSize.Y + Style->MeterPadding.Y));
+		
+		const FVector2D MeterValueSize(ChannelMeterValue * (AllottedWidth - 2.0f * Style->MeterPadding.X), 
+								 Style->MeterSize.Y - 2.0f * Style->MeterValuePadding);
+ 
+		FSlateDrawElement::MakeBox(
+			OutDrawElements,
+			LayerId,
+			MeterGeometry.ToPaintGeometry(MeterValueSize, FSlateLayoutTransform(MeterValueTopLeft)),
+			&Style->MeterValueImage,
+			ESlateDrawEffect::None,
+			Style->MeterValueImage.GetTint(InWidgetStyle) * MeterValueColor.Get().GetColor(InWidgetStyle) * InWidgetStyle.GetColorAndOpacityTint()
+		);
+	}
+	++LayerId;
+
+	// Draw the meter peak value
+	for (int32 ChannelIndex = 0; ChannelIndex < NumChannels; ++ChannelIndex)
+	{	
+		const float ChannelPeakValueDb = ChannelInfos[ChannelIndex].PeakValue;
+		const float ChannelPeakValue = FMath::Clamp((ChannelPeakValueDb - MinValueDb) / ValueDeltaDb, 0.0f, 1.0f);
+
+		const float PeakPixel = ChannelPeakValue * (AllottedWidth - 2.0f * Style->MeterPadding.X);
+
+		FVector2D MeterPeakValueTopLeft(Style->MeterPadding.X + PeakPixel - 0.5f * Style->PeakValueWidth, 
+										Style->MeterPadding.Y + ScaleOffset + Style->MeterValuePadding + ChannelIndex * (Style->MeterSize.Y + Style->MeterPadding.Y));
+		
+		FVector2D MeterPeakValueSize(Style->PeakValueWidth, 
+									 Style->MeterSize.Y - 2.0f * Style->MeterValuePadding);
+
+		FSlateDrawElement::MakeBox(
+			OutDrawElements,
+			LayerId,
+			MeterGeometry.ToPaintGeometry(MeterPeakValueSize, FSlateLayoutTransform(MeterPeakValueTopLeft)),
+			&Style->MeterPeakImage,
+			ESlateDrawEffect::None,
+			Style->MeterPeakImage.GetTint(InWidgetStyle) * MeterPeakColor.Get().GetColor(InWidgetStyle) * InWidgetStyle.GetColorAndOpacityTint()
+		);
+	}
+	++LayerId;
+
+	// Draw the scale hash
+	if (Style->bShowScale)
+	{
+		const float ScaleHashHalfWidth = 0.5f * Style->ScaleHashWidth;
+		const FVector2D HashSize(Style->ScaleHashWidth, Style->ScaleHashHeight);
+
+		const TSharedRef<FSlateFontMeasure> FontMeasureService = FSlateApplication::Get().GetRenderer()->GetFontMeasureService();
+
+		// Measure the min value label size so we can right-justify if needed
+		const FVector2D MinValueLabelSize = FontMeasureService->Measure(FString::FromInt(MinValueDb), Style->Font);
+
+		// Get the size of the negative sign to use to offset the label text in horizontal mode
+		FVector2D NegativeSignSize = FVector2D::ZeroVector;
+		if (Orientation == Orient_Horizontal)
+		{
+			NegativeSignSize = FontMeasureService->Measure(TEXT("-"), Style->Font);
+		}
+
+		const int32 ValueDelta = MaxValueDb - MinValueDb;
+
+		int32 CurrentHashValue = MaxValueDb;
+		while (CurrentHashValue >= MinValueDb)
+		{
+			// Get the fractional value for hash mark
+			const float CurrentHashMeterValuePercent = FMath::Clamp((static_cast<float>(CurrentHashValue) - static_cast<float>(MinValueDb)) / ValueDelta, 0.0f, 1.0f);
+
+			// Draw hash
+			const float HashPixelCenter = CurrentHashMeterValuePercent * (AllottedWidth - 2.0f * Style->MeterPadding.X);
+
+			FVector2D HashTopLeft;
+			HashTopLeft.X = Style->MeterPadding.X + HashPixelCenter - ScaleHashHalfWidth;
+			HashTopLeft.Y = Style->bScaleSide 
+							? Style->MeterPadding.Y + ScaleOffset - Style->ScaleHashOffset - Style->ScaleHashHeight
+							: (Style->MeterSize.Y + Style->MeterPadding.Y) * NumChannels + Style->ScaleHashOffset;
+
+
+			FSlateDrawElement::MakeBox(
+				OutDrawElements,
+				LayerId,
+				MeterGeometry.ToPaintGeometry(HashSize, FSlateLayoutTransform(HashTopLeft)),
+				&Style->MeterPeakImage,
+				ESlateDrawEffect::None,
+				MeterScaleColor.Get().GetColor(InWidgetStyle) * InWidgetStyle.GetColorAndOpacityTint()
+			);
+		
+			const FString LabelString = FString::FromInt(CurrentHashValue);
+			bool bIsNegative = (CurrentHashValue < 0);
+
+			FVector2D LabelSize; 
+			if (Orientation == Orient_Horizontal && bIsNegative)
+			{
+				// We want to center the text on just the positive portion of the number
+				FString LabelStringPositive = FString::FromInt(FMath::Abs(CurrentHashValue));
+				LabelSize = FontMeasureService->Measure(LabelStringPositive, Style->Font);
+			}
+			else
+			{
+				LabelSize = FontMeasureService->Measure(LabelString, Style->Font);
+			}
+			
+			const FText LabelText = FText::FromString(LabelString);
+
+			if (Orientation == Orient_Vertical)
+			{
+				FVector2D LabelTopLeft;
+				
+				LabelTopLeft.X = Style->MeterPadding.X + HashPixelCenter + 0.5f * LabelSize.Y;
+				LabelTopLeft.Y = Style->bScaleSide 
+								 ? Style->MeterPadding.Y - 2.0f + (MinValueLabelSize.X - LabelSize.X)
+								 : (Style->MeterSize.Y + Style->MeterPadding.Y) * NumChannels + Style->ScaleHashOffset + Style->ScaleHashHeight + 2.0f;
+				
+				// Undo the rotation for vertical right before we do the rendering of the scale value 
+				const FSlateRenderTransform RotationTransform = FSlateRenderTransform(FQuat2D(FMath::DegreesToRadians(90.0f)));
+
+				const FSlateLayoutTransform ChildLayoutTransform(1.0f, TransformPoint(1.0f, LabelTopLeft));
+
+				const FGeometry TextGeometry = MeterGeometry.MakeChild(
+					LabelSize,
+					ChildLayoutTransform,
+					RotationTransform,
+					FVector2D(0.0f, 0.0f)
+				);
+
+				FSlateDrawElement::MakeText(
+					OutDrawElements,
+					LayerId,
+					TextGeometry.ToOffsetPaintGeometry(FVector2D(0.0f, 0.0f)),
+					LabelText,
+					Style->Font,
+					ESlateDrawEffect::None,
+					MeterScaleLabelColor.Get().GetColor(InWidgetStyle)* InWidgetStyle.GetColorAndOpacityTint()
+				);
+			}
+			else
+			{
+				FVector2D LabelTopLeft;
+
+				// Center it based off the positive value width, then subtract from there the size of the negative sign
+				LabelTopLeft.X = Style->MeterPadding.X + HashPixelCenter - 0.5f * LabelSize.X;
+				if (bIsNegative)
+				{
+					LabelTopLeft.X -= NegativeSignSize.X;
+				}
+
+				LabelTopLeft.Y = Style->bScaleSide 
+								 ? Style->MeterPadding.Y + ScaleOffset - LabelSize.Y - Style->ScaleHashOffset - Style->ScaleHashHeight
+								 : (Style->MeterSize.Y + Style->MeterPadding.Y) * NumChannels + Style->ScaleHashOffset + Style->ScaleHashHeight + 2.0f;
+
+				// Draw text label
+				FSlateDrawElement::MakeText(
+					OutDrawElements,
+					LayerId,
+					MeterGeometry.ToOffsetPaintGeometry(LabelTopLeft),
+					LabelText,
+					Style->Font,
+					ESlateDrawEffect::None,
+					MeterScaleLabelColor.Get().GetColor(InWidgetStyle) * InWidgetStyle.GetColorAndOpacityTint()
+				);
+
+			}
+
+			CurrentHashValue -= Style->DecibelsPerHash;
+		}
+	}
+
+	return LayerId;
+}
+
+float SAudioMeterWidget::GetScaleHeight() const
+{
+	const TSharedRef<FSlateFontMeasure> FontMeasureService = FSlateApplication::Get().GetRenderer()->GetFontMeasureService();
+	FVector2D LabelSize = FontMeasureService->Measure(FString::FromInt(-60), Style->Font);
+
+	float ScaleHeight = Style->ScaleHashHeight + Style->ScaleHashOffset;
+
+	if (Orientation == Orient_Vertical)
+	{
+		return ScaleHeight + LabelSize.X;
+	}
+	return ScaleHeight + LabelSize.Y;
+}
+
+FVector2D SAudioMeterWidget::ComputeDesiredSize(float) const
+{
+	static const FVector2D SAudioMeterDesiredSize(50.0f, 50.0f);
+
+	if (Style == nullptr)
+	{
+		return SAudioMeterDesiredSize;
+	}
+
+	TArray<FAudioMeterChannelInfo> ChannelInfo = MeterChannelInfoAttribute.Get();
+	const int32 NumChannels = FMath::Max(ChannelInfo.Num(), 1);
+
+	FVector2D Size = FVector2D(Style->MeterSize.X, (Style->MeterSize.Y + Style->MeterPadding.Y) * NumChannels);
+
+	// Add the end padding
+	Size.X += Style->MeterPadding.X;
+	Size.Y += Style->MeterPadding.Y;
+
+	// Add the width for the scale if it's been set to show
+	if (Style->bShowScale)
+	{
+		Size.Y += GetScaleHeight();
+	}
+
+	// The thickness is going to be the thickness plus the scale image width
+	if (Orientation == Orient_Vertical)
+	{
+		return FVector2D(Size.Y, Size.X);
+	}
+
+	return Size; 
+}
+
+bool SAudioMeterWidget::ComputeVolatility() const
+{
+	return true; 
+}
+
+void SAudioMeterWidget::SetMeterChannelInfo(const TAttribute<TArray<FAudioMeterChannelInfo>>& InMeterChannelInfo)
+{
+	SetAttribute(MeterChannelInfoAttribute, InMeterChannelInfo, EInvalidateWidgetReason::Paint);
+}
+
+TArray<FAudioMeterChannelInfo> SAudioMeterWidget::GetMeterChannelInfo() const
+{
+	return MeterChannelInfoAttribute.Get();
+}
+
+void SAudioMeterWidget::SetOrientation(EOrientation InOrientation)
+{
+	if (Orientation != InOrientation)
+	{
+		Orientation = InOrientation;
+		Invalidate(EInvalidateWidgetReason::Layout);
+	}
+}
+
+void SAudioMeterWidget::SetBackgroundColor(FSlateColor InBackgroundColor)
+{
+	SetAttribute(BackgroundColor, TAttribute<FSlateColor>(InBackgroundColor), EInvalidateWidgetReason::Paint);
+}
+
+void SAudioMeterWidget::SetMeterBackgroundColor(FSlateColor InMeterBackgroundColor)
+{
+	SetAttribute(MeterBackgroundColor, TAttribute<FSlateColor>(InMeterBackgroundColor), EInvalidateWidgetReason::Paint);
+}
+
+void SAudioMeterWidget::SetMeterValueColor(FSlateColor InMeterValueColor)
+{
+	SetAttribute(MeterValueColor, TAttribute<FSlateColor>(InMeterValueColor), EInvalidateWidgetReason::Paint);
+}
+
+void SAudioMeterWidget::SetMeterPeakColor(FSlateColor InMeterPeakColor)
+{
+	SetAttribute(MeterPeakColor, TAttribute<FSlateColor>(InMeterPeakColor), EInvalidateWidgetReason::Paint);
+}
+
+void SAudioMeterWidget::SetMeterClippingColor(FSlateColor InMeterClippingColor)
+{
+	SetAttribute(MeterClippingColor, TAttribute<FSlateColor>(InMeterClippingColor), EInvalidateWidgetReason::Paint);
+}
+
+void SAudioMeterWidget::SetMeterScaleColor(FSlateColor InMeterScaleColor)
+{
+	SetAttribute(MeterScaleColor, TAttribute<FSlateColor>(InMeterScaleColor), EInvalidateWidgetReason::Paint);
+}
+
+void SAudioMeterWidget::SetMeterScaleLabelColor(FSlateColor InMeterScaleLabelColor)
+{
+	SetAttribute(MeterScaleLabelColor, TAttribute<FSlateColor>(InMeterScaleLabelColor), EInvalidateWidgetReason::Paint);
+}
diff --git a/Engine/Plugins/Runtime/AudioWidgets/Source/AudioWidgetsCore/Public/AudioMeterChannelInfo.h b/Engine/Plugins/Runtime/AudioWidgets/Source/AudioWidgetsCore/Public/AudioMeterChannelInfo.h
new file mode 100644
--- /dev/null
+++ b/Engine/Plugins/Runtime/AudioWidgets/Source/AudioWidgetsCore/Public/AudioMeterChannelInfo.h
@@ -0,0 +1,27 @@
+// Copyright Epic Games, Inc. All Rights Reserved.
+
+#pragma once
+
+#include "AudioMeterChannelInfo.generated.h"
+
+USTRUCT(BlueprintType)
+struct FAudioMeterChannelInfo
+{
+	GENERATED_BODY()
+
+	UPROPERTY(EditAnywhere, BlueprintReadWrite, Category = AudioMeter)
+	float MeterValue = 0.0f;
+
+	UPROPERTY(EditAnywhere, BlueprintReadWrite, Category = AudioMeter)
+	float PeakValue = 0.0f;
+
+	UPROPERTY(EditAnywhere, BlueprintReadWrite, Category = AudioMeter)
+	float ClippingValue = 0.0f;
+};
+
+inline bool operator==(const FAudioMeterChannelInfo& Lhs, const FAudioMeterChannelInfo& Rhs)
+{
+	return FMath::IsNearlyEqual(Lhs.MeterValue, Rhs.MeterValue) &&
+		   FMath::IsNearlyEqual(Lhs.PeakValue, Rhs.PeakValue)   &&
+		   FMath::IsNearlyEqual(Lhs.ClippingValue, Rhs.ClippingValue);
+}
diff --git a/Engine/Plugins/Runtime/AudioWidgets/Source/AudioWidgetsCore/Public/AudioMeterWidgetStyle.h b/Engine/Plugins/Runtime/AudioWidgets/Source/AudioWidgetsCore/Public/AudioMeterWidgetStyle.h
new file mode 100644
--- /dev/null
+++ b/Engine/Plugins/Runtime/AudioWidgets/Source/AudioWidgetsCore/Public/AudioMeterWidgetStyle.h
@@ -0,0 +1,172 @@
+// Copyright Epic Games, Inc. All Rights Reserved.
+
+#pragma once
+
+#include "Fonts/SlateFontInfo.h"
+#include "Styling/SlateBrush.h"
+#include "Styling/SlateWidgetStyle.h"
+#include "Styling/StyleDefaults.h"
+
+#include "AudioMeterWidgetStyle.generated.h"
+
+#define UE_API AUDIOWIDGETSCORE_API
+
+struct FCompositeFont;
+
+/**
+ * Represents the appearance of an AudioMeter widget
+ */
+USTRUCT(BlueprintType)
+struct FAudioMeterWidgetStyle : public FSlateWidgetStyle
+{
+	GENERATED_BODY()
+
+	UE_API static const FAudioMeterWidgetStyle& GetDefault();
+
+	UE_API virtual void GetResources(TArray<const FSlateBrush*>& OutBrushes) const override;
+	virtual const FName GetTypeName() const override { return TypeName; };
+
+	void UnlinkColors();
+
+	inline static const FName TypeName = "FAudioMeterWidgetStyle";
+
+	// Image to use to represent the meter value.
+	UPROPERTY(EditAnywhere, BlueprintReadWrite, Category=Appearance)
+	FSlateBrush MeterValueImage;
+	FAudioMeterWidgetStyle& SetMeterValueImage(const FSlateBrush& InMeterValueImage){ MeterValueImage = InMeterValueImage; return *this; }
+
+	// Image to use to represent the background.
+	UPROPERTY(EditAnywhere, BlueprintReadWrite, Category = Appearance)
+	FSlateBrush BackgroundImage;
+	FAudioMeterWidgetStyle& SetBackgroundImage(const FSlateBrush& InBackgroundImage) { BackgroundImage = InBackgroundImage; return *this; }
+
+	// Image to use to represent the meter background.
+	UPROPERTY(EditAnywhere, BlueprintReadWrite, Category = Appearance)
+	FSlateBrush MeterBackgroundImage;
+	FAudioMeterWidgetStyle& SetMeterBackgroundImage(const FSlateBrush& InMeterBackgroundImage) { MeterBackgroundImage = InMeterBackgroundImage; return *this; }
+
+	// Image to use to draw behind the meter value.
+	UPROPERTY(EditAnywhere, BlueprintReadWrite, Category = Appearance)
+	FSlateBrush MeterValueBackgroundImage;
+	FAudioMeterWidgetStyle& SetMeterValueBackgroundImage(const FSlateBrush& InMeterValueBackgroundImage) { MeterValueBackgroundImage = InMeterValueBackgroundImage; return *this; }
+
+	// Image to use to represent the meter peak.
+	UPROPERTY(EditAnywhere, BlueprintReadWrite, Category = Appearance)
+	FSlateBrush MeterPeakImage;
+	FAudioMeterWidgetStyle& SetMeterPeakImage(const FSlateBrush& InMeterPeakImage) { MeterPeakImage = InMeterPeakImage; return *this; }
+
+	// How thick to draw the meter
+	UPROPERTY(EditAnywhere, BlueprintReadWrite, Category=Appearance)
+	FVector2D MeterSize = FVector2D(250.0f, 25.0f);
+	FAudioMeterWidgetStyle& SetMeterSize(const FVector2D& InMeterSize) { MeterSize = InMeterSize; return *this; }
+
+	// How much padding to add around the meter
+	UPROPERTY(EditAnywhere, BlueprintReadWrite, Category = Appearance)
+	FVector2D MeterPadding = FVector2D(10.0f, 5.0f);
+	FAudioMeterWidgetStyle& SetMeterPadding(const FVector2D& InMeterPadding) { MeterPadding = InMeterPadding; return *this; }
+
+	// How much padding to add around the meter value
+	UPROPERTY(EditAnywhere, BlueprintReadWrite, Category = Appearance)
+	float MeterValuePadding = 3.0f;
+	FAudioMeterWidgetStyle& SetMeterValuePadding(float InMeterValuePadding) { MeterValuePadding = InMeterValuePadding; return *this; }
+
+	// How wide to draw the peak value indicator
+	UPROPERTY(EditAnywhere, BlueprintReadWrite, Category = Appearance)
+	float PeakValueWidth = 2.0f;
+	FAudioMeterWidgetStyle& SetPeakValueWidth(float InPeakValueWidth) { PeakValueWidth = InPeakValueWidth; return *this; }
+
+	// The minimum and maximum value to display in dB (values are clamped in this range)
+	UPROPERTY(EditAnywhere, BlueprintReadWrite, Category = Appearance)
+	FVector2D ValueRangeDb = FVector2D(-160, 10);
+	FAudioMeterWidgetStyle& SetValueRangeDb(const FVector2D& InValueRangeDb) { ValueRangeDb = InValueRangeDb; return *this; }
+
+	// Whether or not to show the decibel scale alongside the meter
+	UPROPERTY(EditAnywhere, BlueprintReadWrite, Category = Appearance)
+	bool bShowScale = true;
+	FAudioMeterWidgetStyle& SetShowScale(bool bInShowScale) { bShowScale = bInShowScale; return *this; }
+
+	// Which side to show the scale. If vertical, true means left side, false means right side. If horizontal, true means above, false means below.
+	UPROPERTY(EditAnywhere, BlueprintReadWrite, Category = Appearance)
+	bool bScaleSide = true;
+	FAudioMeterWidgetStyle& SetScaleSide(bool bInScaleSide) { bScaleSide = bInScaleSide; return *this; }
+
+	// Offset for the hashes
+	UPROPERTY(EditAnywhere, BlueprintReadWrite, Category = Appearance)
+	float ScaleHashOffset = 5.0f;
+	FAudioMeterWidgetStyle& SetScaleHashOffset(float InScaleHashOffset) { ScaleHashOffset = InScaleHashOffset; return *this; }
+
+	// The width of each hash mark
+	UPROPERTY(EditAnywhere, BlueprintReadWrite, Category = Appearance)
+	float ScaleHashWidth = 1.0f;
+	FAudioMeterWidgetStyle& SetScaleHashWidth(float InScaleHashWidth) { ScaleHashWidth = InScaleHashWidth; return *this; }
+
+	// The height of each hash mark
+	UPROPERTY(EditAnywhere, BlueprintReadWrite, Category = Appearance)
+	float ScaleHashHeight = 10.0f;
+	FAudioMeterWidgetStyle& SetScaleHashHeight(float InScaleHashHeight) { ScaleHashHeight = InScaleHashHeight; return *this; }
+
+	// How wide to draw the decibel scale, if it's enabled
+	UPROPERTY(EditAnywhere, BlueprintReadWrite, Category = Appearance, meta = (UIMin = "3", ClampMin="3", UIMax = "10"))
+	int32 DecibelsPerHash = 10;
+	FAudioMeterWidgetStyle& SetDecibelsPerHash(float InDecibelsPerHash) { DecibelsPerHash = InDecibelsPerHash; return *this; }
+
+	/** Font family and size to be used when displaying the meter scale. */
+	UPROPERTY(EditAnywhere, BlueprintReadWrite, Category = Appearance)
+	FSlateFontInfo Font = FStyleDefaults::GetFontInfo(5);
+	FAudioMeterWidgetStyle& SetFont(const FSlateFontInfo& InFont) { Font = InFont; return *this; }
+	FAudioMeterWidgetStyle& SetFont(TSharedPtr<const FCompositeFont> InCompositeFont, const int32 InSize, const FName& InTypefaceFontName = NAME_None) { Font = FSlateFontInfo(InCompositeFont, InSize, InTypefaceFontName); return *this; }
+	FAudioMeterWidgetStyle& SetFont(const UObject* InFontObject, const int32 InSize, const FName& InTypefaceFontName = NAME_None) { Font = FSlateFontInfo(InFontObject, InSize, InTypefaceFontName); return *this; }
+
+	UE_DEPRECATED(5.6, "Use SetFont using FSlateFontInfo instead. FSlateFontInfo's constructors using a FontName are deprecated.")
+	UE_API FAudioMeterWidgetStyle& SetFont(const FName& InFontName, uint16 InSize);
+	UE_DEPRECATED(5.6, "Use SetFont using FSlateFontInfo instead. FSlateFontInfo's constructors using a FontName are deprecated.")
+	UE_API FAudioMeterWidgetStyle& SetFont(const FString& InFontName, uint16 InSize);
+	UE_DEPRECATED(5.6, "Use SetFont using FSlateFontInfo instead. FSlateFontInfo's constructors using a FontName are deprecated.")
+	UE_API FAudioMeterWidgetStyle& SetFont(const WIDECHAR* InFontName, uint16 InSize);
+	UE_DEPRECATED(5.6, "Use SetFont using FSlateFontInfo instead. FSlateFontInfo's constructors using a FontName are deprecated.")
+	UE_API FAudioMeterWidgetStyle& SetFont(const ANSICHAR* InFontName, uint16 InSize);
+	UE_DEPRECATED(5.6, "Use SetFont using FSlateFontInfo instead. FSlateFontInfo's constructors using a FontName are deprecated.")
+	UE_API FAudioMeterWidgetStyle& SetFontName(const FName& InFontName);
+	UE_DEPRECATED(5.6, "Use SetFont using FSlateFontInfo instead. FSlateFontInfo's constructors using a FontName are deprecated.")
+	UE_API FAudioMeterWidgetStyle& SetFontName(const FString& InFontName);
+	UE_DEPRECATED(5.6, "Use SetFont using FSlateFontInfo instead. FSlateFontInfo's constructors using a FontName are deprecated.")
+	UE_API FAudioMeterWidgetStyle& SetFontName(const WIDECHAR* InFontName);
+	UE_DEPRECATED(5.6, "Use SetFont using FSlateFontInfo instead. FSlateFontInfo's constructors using a FontName are deprecated.")
+	UE_API FAudioMeterWidgetStyle& SetFontName(const ANSICHAR* InFontName);
+
+	FAudioMeterWidgetStyle& SetFontSize(uint16 InSize) { Font.Size = InSize; return *this; }
+	FAudioMeterWidgetStyle& SetTypefaceFontName(const FName& InTypefaceFontName) { Font.TypefaceFontName = InTypefaceFontName; return *this; }
+};
+
+
+USTRUCT(BlueprintType)
+struct FAudioMeterDefaultColorWidgetStyle : public FSlateWidgetStyle
+{
+	GENERATED_BODY()
+
+	UE_API static const FAudioMeterDefaultColorWidgetStyle& GetDefault();
+
+	virtual const FName GetTypeName() const override { return TypeName; };
+
+	inline static const FName TypeName = "FAudioMeterDefaultColorWidgetStyle";
+
+	UPROPERTY(EditAnywhere, BlueprintReadWrite, Category = "Style")
+	FLinearColor MeterBackgroundColor = FLinearColor(0.031f, 0.031f, 0.031f, 1.0f);
+
+	UPROPERTY(EditAnywhere, BlueprintReadWrite, Category = "Style")
+	FLinearColor MeterValueColor = FLinearColor(0.025719f, 0.208333f, 0.069907f, 1.0f);
+
+	UPROPERTY(EditAnywhere, BlueprintReadWrite, Category = "Style")
+	FLinearColor MeterPeakColor = FLinearColor(0.24349f, 0.708333f, 0.357002f, 1.0f);
+
+	UPROPERTY(EditAnywhere, BlueprintReadWrite, Category = "Style")
+	FLinearColor MeterClippingColor = FLinearColor(1.0f, 0.0f, 0.112334f, 1.0f);
+
+	UPROPERTY(EditAnywhere, BlueprintReadWrite, Category = "Style")
+	FLinearColor MeterScaleColor = FLinearColor(0.017642f, 0.017642f, 0.017642f, 1.0f);
+
+	UPROPERTY(EditAnywhere, BlueprintReadWrite, Category = "Style")
+	FLinearColor MeterScaleLabelColor = FLinearColor(0.442708f, 0.442708f, 0.442708f, 1.0f);
+};
+
+#undef UE_API
diff --git a/Engine/Plugins/Runtime/AudioWidgets/Source/AudioWidgetsCore/Public/SAudioMeterWidget.h b/Engine/Plugins/Runtime/AudioWidgets/Source/AudioWidgetsCore/Public/SAudioMeterWidget.h
new file mode 100644
--- /dev/null
+++ b/Engine/Plugins/Runtime/AudioWidgets/Source/AudioWidgetsCore/Public/SAudioMeterWidget.h
@@ -0,0 +1,119 @@
+// Copyright Epic Games, Inc. All Rights Reserved.
+
+#pragma once
+
+#include "AudioMeterChannelInfo.h"
+#include "AudioMeterWidgetStyle.h"
+#include "Math/Color.h"
+#include "Styling/ISlateStyle.h"
+#include "Styling/SlateWidgetStyleAsset.h"
+#include "Widgets/DeclarativeSyntaxSupport.h"
+#include "Widgets/SLeafWidget.h"
+
+#define UE_API AUDIOWIDGETSCORE_API
+
+class FPaintArgs;
+class FSlateWindowElementList;
+
+class SAudioMeterWidgetBase : public SLeafWidget
+{
+public:
+	SLATE_BEGIN_ARGS(SAudioMeterWidgetBase){}
+	SLATE_END_ARGS()
+
+	UE_API virtual TArray<FAudioMeterChannelInfo> GetMeterChannelInfo() const = 0;
+	UE_API virtual void SetMeterChannelInfo(const TAttribute<TArray<FAudioMeterChannelInfo>>& InMeterChannelInfo) = 0;
+
+	// The active timer registered to refresh the meter channel info
+	bool bIsActiveTimerRegistered = false;
+};
+
+/**
+ * A widget that displays audio channel's levels.
+ */
+class SAudioMeterWidget : public SAudioMeterWidgetBase
+{
+public:
+	SLATE_BEGIN_ARGS(SAudioMeterWidget)
+		: _Orientation(EOrientation::Orient_Horizontal)
+		, _BackgroundColor(FLinearColor::Black)
+		, _MeterBackgroundColor(FLinearColor::Gray)
+		, _MeterValueColor(FLinearColor::Green)
+		, _MeterPeakColor(FLinearColor::Blue)
+		, _MeterScaleColor(FLinearColor::White)
+		, _MeterScaleLabelColor(FLinearColor::Gray)
+		, _Style(&FAudioMeterWidgetStyle::GetDefault())
+	{
+	}
+
+		// The audio meter's orientation
+		SLATE_ARGUMENT(EOrientation, Orientation)
+
+		// The color to draw the background in
+		SLATE_ATTRIBUTE(FSlateColor, BackgroundColor)
+
+		// The color to draw the meter background in
+		SLATE_ATTRIBUTE(FSlateColor, MeterBackgroundColor)
+
+		// The color to draw the meter value in
+		SLATE_ATTRIBUTE(FSlateColor, MeterValueColor)
+
+		// The color to draw the meter peak
+		SLATE_ATTRIBUTE(FSlateColor, MeterPeakColor)
+
+		// The color to draw the clipping value in
+		SLATE_ATTRIBUTE(FSlateColor, MeterClippingColor)
+
+		// The color to draw the scale in
+		SLATE_ATTRIBUTE(FSlateColor, MeterScaleColor)
+
+		// The color to draw the scale in
+		SLATE_ATTRIBUTE(FSlateColor, MeterScaleLabelColor)
+
+		// The style used to draw the slider
+		SLATE_STYLE_ARGUMENT(FAudioMeterWidgetStyle, Style)
+
+		// A value representing the audio meter value
+		SLATE_ATTRIBUTE(TArray<FAudioMeterChannelInfo>, MeterChannelInfo)
+
+	SLATE_END_ARGS()
+
+	UE_API void Construct(const SAudioMeterWidget::FArguments& InDeclaration);
+
+	UE_API virtual int32 OnPaint(const FPaintArgs& Args, const FGeometry& AllottedGeometry, const FSlateRect& MyCullingRect, FSlateWindowElementList& OutDrawElements, int32 LayerId, const FWidgetStyle& InWidgetStyle, bool bParentEnabled) const override;
+	UE_API virtual FVector2D ComputeDesiredSize(float) const override;
+	UE_API virtual bool ComputeVolatility() const override;
+
+	UE_API virtual TArray<FAudioMeterChannelInfo> GetMeterChannelInfo() const override;
+	UE_API virtual void SetMeterChannelInfo(const TAttribute<TArray<FAudioMeterChannelInfo>>& InMeterChannelInfo) override;
+
+	UE_API void SetOrientation(EOrientation InOrientation);
+	UE_API void SetBackgroundColor(FSlateColor InBackgroundColor);
+	UE_API void SetMeterBackgroundColor(FSlateColor InMeterBackgroundColor);
+	UE_API void SetMeterValueColor(FSlateColor InMeterValueColor);
+	UE_API void SetMeterPeakColor(FSlateColor InMeterPeakColor);
+	UE_API void SetMeterClippingColor(FSlateColor InMeterPeakColor);
+	UE_API void SetMeterScaleColor(FSlateColor InMeterScaleColor);
+	UE_API void SetMeterScaleLabelColor(FSlateColor InMeterScaleLabelColor);
+
+	UE_API float GetScaleHeight() const;
+
+protected:
+	const FAudioMeterWidgetStyle* Style = nullptr;
+
+	// Audio meter's orientation
+	EOrientation Orientation = EOrientation::Orient_Vertical;
+
+	// Colors
+	TAttribute<FSlateColor> BackgroundColor;
+	TAttribute<FSlateColor> MeterBackgroundColor;
+	TAttribute<FSlateColor> MeterValueColor;
+	TAttribute<FSlateColor> MeterPeakColor;
+	TAttribute<FSlateColor> MeterClippingColor;
+	TAttribute<FSlateColor> MeterScaleColor;
+	TAttribute<FSlateColor> MeterScaleLabelColor;
+
+	TAttribute<TArray<FAudioMeterChannelInfo>> MeterChannelInfoAttribute;
+};
+
+#undef UE_API
