[AudioWidgets] LKFS Loudness Meter added to analyzer rack.

diff --git a/Engine/Plugins/Runtime/AudioWidgets/Source/AudioWidgets/Private/AudioLoudnessMeter.cpp b/Engine/Plugins/Runtime/AudioWidgets/Source/AudioWidgets/Private/AudioLoudnessMeter.cpp
new file mode 100644
--- /dev/null
+++ b/Engine/Plugins/Runtime/AudioWidgets/Source/AudioWidgets/Private/AudioLoudnessMeter.cpp
@@ -0,0 +1,140 @@
+// Copyright Epic Games, Inc. All Rights Reserved.
+
+#include "AudioLoudnessMeter.h"
+
+#include "Widgets/Docking/SDockTab.h"
+
+#define LOCTEXT_NAMESPACE "FAudioLoudnessMeter"
+
+namespace AudioWidgets
+{
+	const FAudioAnalyzerRackUnitTypeInfo FAudioLoudnessMeter::RackUnitTypeInfo
+	{
+		.TypeName = TEXT("FAudioLoudnessMeter"),
+		.DisplayName = LOCTEXT("AudioLoudnessMeterDisplayName", "Loudness"),
+		.OnMakeAudioAnalyzerRackUnit = FOnMakeAudioAnalyzerRackUnit::CreateStatic(&MakeRackUnit),
+		.VerticalSizeCoefficient = 0.5f,
+	};
+
+	FAudioLoudnessMeter::FAudioLoudnessMeter(const FAudioLoudnessMeterParams& Params)
+	{
+		LoudnessMeterWidgetView.AddNumericValue(
+			{
+				.Name = GET_MEMBER_NAME_CHECKED(FLKFSResults, IntegratedLoudness),
+				.DisplayName = LOCTEXT("IntegratedLoudnessDisplayName", "Integrated"),
+				.Value = TAttribute<TOptional<float>>::CreateLambda([this]() { return (LatestLoudnessResults.IsSet()) ? TOptional<float>(LatestLoudnessResults->IntegratedLoudness) : NullOpt; })
+			});
+
+		LoudnessMeterWidgetView.AddMeterValue(
+			{
+				.Name = GET_MEMBER_NAME_CHECKED(FLKFSResults, ShortTermLoudness),
+				.DisplayName = LOCTEXT("ShortTermLoudnessDisplayName", "Short Term"),
+				.Value = TAttribute<TOptional<float>>::CreateLambda([this]() { return (LatestLoudnessResults.IsSet()) ? TOptional<float>(LatestLoudnessResults->ShortTermLoudness) : NullOpt; })
+			});
+		LoudnessMeterWidgetView.AddMeterValue(
+			{
+				.Name = GET_MEMBER_NAME_CHECKED(FLKFSResults, Loudness),
+				.DisplayName = LOCTEXT("MomentaryLoudnessDisplayName", "Momentary"),
+				.Value = TAttribute<TOptional<float>>::CreateLambda([this]() { return (LatestLoudnessResults.IsSet()) ? TOptional<float>(LatestLoudnessResults->Loudness) : NullOpt; })
+			});
+
+		Init(Params.NumChannels, Params.AudioDeviceId, Params.ExternalAudioBus);
+	}
+
+	FAudioLoudnessMeter::~FAudioLoudnessMeter()
+	{
+		Teardown();
+	}
+
+	void FAudioLoudnessMeter::SetAudioBusInfo(const FAudioBusInfo& AudioBusInfo)
+	{
+		Init(AudioBusInfo.AudioBus->GetNumChannels(), AudioBusInfo.AudioDeviceId, AudioBusInfo.AudioBus);
+	}
+
+	TSharedRef<SDockTab> FAudioLoudnessMeter::SpawnTab(const FSpawnTabArgs& Args) const
+	{
+		return SNew(SDockTab)
+			.Clipping(EWidgetClipping::ClipToBounds)
+			.Label(RackUnitTypeInfo.DisplayName)
+			[
+				LoudnessMeterWidgetView.MakeWidget()
+			];
+	}
+
+	TSharedRef<IAudioAnalyzerRackUnit> FAudioLoudnessMeter::MakeRackUnit(const FAudioAnalyzerRackUnitConstructParams& Params)
+	{
+		const FAudioLoudnessMeterParams AnalyzerParams
+		{
+			.NumChannels = Params.AudioBusInfo.GetNumChannels(),
+			.AudioDeviceId = Params.AudioBusInfo.AudioDeviceId,
+			.ExternalAudioBus = Params.AudioBusInfo.AudioBus,
+		};
+
+		return MakeShared<FAudioLoudnessMeter>(AnalyzerParams);
+	}
+
+	void FAudioLoudnessMeter::Init(const int32 InNumChannels, const Audio::FDeviceId InAudioDeviceId, const TObjectPtr<UAudioBus> InExternalAudioBus)
+	{
+		Teardown();
+
+		LKFSSettings = TStrongObjectPtr(NewObject<ULKFSSettings>());
+
+		// Only create analyzers etc if we have an audio device:
+		if (InAudioDeviceId != FAudioBusInfo::InvalidAudioDeviceId)
+		{
+			check(InNumChannels > 0);
+
+			bUseExternalAudioBus = (InExternalAudioBus != nullptr);
+			AudioBus = (bUseExternalAudioBus) ? TStrongObjectPtr(InExternalAudioBus.Get()) : TStrongObjectPtr(NewObject<UAudioBus>());
+			AudioBus->AudioBusChannels = EAudioBusChannels(InNumChannels - 1);
+
+			CreateLKFSAnalyzer();
+
+			LKFSAnalyzer->StartAnalyzing(InAudioDeviceId, AudioBus.Get());
+		}
+	}
+
+	void FAudioLoudnessMeter::Teardown()
+	{
+		if (LKFSAnalyzer.IsValid() && LKFSAnalyzer->IsValidLowLevel())
+		{
+			LKFSAnalyzer->StopAnalyzing();
+			ReleaseLKFSAnalyzer();
+		}
+
+		AudioBus.Reset();
+		LKFSSettings.Reset();
+		bUseExternalAudioBus = false;
+	}
+
+	void FAudioLoudnessMeter::CreateLKFSAnalyzer()
+	{
+		ensure(!LKFSAnalyzer.IsValid());
+		ensure(!LatestOverallLKFSResultsDelegateHandle.IsValid());
+
+		LKFSAnalyzer = TStrongObjectPtr(NewObject<ULKFSAnalyzer>());
+		LKFSAnalyzer->Settings = LKFSSettings.Get();
+		LatestOverallLKFSResultsDelegateHandle = LKFSAnalyzer->OnLatestOverallLKFSResultsNative.AddRaw(this, &FAudioLoudnessMeter::OnLoudnessOutput);
+	}
+
+	void FAudioLoudnessMeter::ReleaseLKFSAnalyzer()
+	{
+		if (ensure(LKFSAnalyzer.IsValid() && LatestOverallLKFSResultsDelegateHandle.IsValid()))
+		{
+			LKFSAnalyzer->OnLatestOverallLKFSResultsNative.Remove(LatestOverallLKFSResultsDelegateHandle);
+		}
+
+		LatestOverallLKFSResultsDelegateHandle.Reset();
+		LKFSAnalyzer.Reset();
+	}
+
+	void FAudioLoudnessMeter::OnLoudnessOutput(ULKFSAnalyzer* InLKFSAnalyzer, const FLKFSResults& InLoudnessResults)
+	{
+		if (InLKFSAnalyzer == LKFSAnalyzer.Get())
+		{
+			LatestLoudnessResults = InLoudnessResults;
+		}
+	}
+} // namespace AudioWidgets
+
+#undef LOCTEXT_NAMESPACE
diff --git a/Engine/Plugins/Runtime/AudioWidgets/Source/AudioWidgets/Private/AudioLoudnessMeter.h b/Engine/Plugins/Runtime/AudioWidgets/Source/AudioWidgets/Private/AudioLoudnessMeter.h
new file mode 100644
--- /dev/null
+++ b/Engine/Plugins/Runtime/AudioWidgets/Source/AudioWidgets/Private/AudioLoudnessMeter.h
@@ -0,0 +1,67 @@
+// Copyright Epic Games, Inc. All Rights Reserved.
+#pragma once
+
+#include "AudioAnalyzerRack.h"
+#include "LKFS.h"
+#include "LoudnessMeterWidgetView.h"
+
+namespace AudioWidgets
+{
+	/**
+	 * Constructor parameters for the analyzer.
+	 */
+	struct FAudioLoudnessMeterParams
+	{
+		int32 NumChannels = 1;
+		Audio::FDeviceId AudioDeviceId = INDEX_NONE;
+		TObjectPtr<UAudioBus> ExternalAudioBus = nullptr;
+	};
+
+	/**
+	 * Owns an analyzer and a corresponding Slate widget for displaying the loudness stats.
+	 * Can either create an Audio Bus to analyze, or analyze the given Bus.
+	 */
+	class FAudioLoudnessMeter : public IAudioAnalyzerRackUnit
+	{
+	public:
+		static const FAudioAnalyzerRackUnitTypeInfo RackUnitTypeInfo;
+		
+		FAudioLoudnessMeter(const FAudioLoudnessMeterParams& Params);
+		~FAudioLoudnessMeter();
+
+		// Begin IAudioAnalyzerRackUnit overrides.
+		virtual void SetAudioBusInfo(const FAudioBusInfo& AudioBusInfo) override;
+		virtual TSharedRef<SDockTab> SpawnTab(const FSpawnTabArgs& Args) const override;
+		// End IAudioAnalyzerRackUnit overrides.
+
+	private:
+		static TSharedRef<IAudioAnalyzerRackUnit> MakeRackUnit(const FAudioAnalyzerRackUnitConstructParams& Params);
+
+		void Init(int32 InNumChannels, Audio::FDeviceId InAudioDeviceId, TObjectPtr<UAudioBus> InExternalAudioBus);
+		void Teardown();
+
+		void CreateLKFSAnalyzer();
+		void ReleaseLKFSAnalyzer();
+
+		void OnLoudnessOutput(ULKFSAnalyzer* InLKFSAnalyzer, const FLKFSResults& InLoudnessResults);
+
+		/** Creates the meters widget */
+		FLoudnessMeterWidgetView LoudnessMeterWidgetView;
+		
+		/** Audio analyzer object. */
+		TStrongObjectPtr<ULKFSAnalyzer> LKFSAnalyzer;
+
+		/** Analyzer settings. */
+		TStrongObjectPtr<ULKFSSettings> LKFSSettings;
+
+		/** Handle for results delegate for analyzer. */
+		FDelegateHandle LatestOverallLKFSResultsDelegateHandle;
+
+		/** Most recent results from the analyzer. */
+		TOptional<FLKFSResults> LatestLoudnessResults;
+
+		/** The audio bus used for analysis. */
+		TStrongObjectPtr<UAudioBus> AudioBus;
+		bool bUseExternalAudioBus = false;
+	};
+} // namespace AudioWidgets
diff --git a/Engine/Plugins/Runtime/AudioWidgets/Source/AudioWidgets/Private/AudioWidgets.cpp b/Engine/Plugins/Runtime/AudioWidgets/Source/AudioWidgets/Private/AudioWidgets.cpp
--- a/Engine/Plugins/Runtime/AudioWidgets/Source/AudioWidgets/Private/AudioWidgets.cpp
+++ b/Engine/Plugins/Runtime/AudioWidgets/Source/AudioWidgets/Private/AudioWidgets.cpp
@@ -3,6 +3,7 @@
 #include "AudioWidgets.h"
 #include "AdvancedWidgetsModule.h"
 #include "AudioAnalyzerRackUnitRegistry.h"
+#include "AudioLoudnessMeter.h"
 #include "AudioMeter.h"
 #include "AudioOscilloscope.h"
 #include "AudioSpectrogram.h"
@@ -24,6 +25,7 @@ void FAudioWidgetsModule::StartupModule()
 	FAudioWidgetsStyle::Get();
 
 	// Register standard analyzer rack units:
+	RegisterAudioAnalyzerRackUnitType(&FAudioLoudnessMeter::RackUnitTypeInfo);
 	RegisterAudioAnalyzerRackUnitType(&FAudioMeter::RackUnitTypeInfo);
 	RegisterAudioAnalyzerRackUnitType(&FAudioOscilloscope::RackUnitTypeInfo);
 	RegisterAudioAnalyzerRackUnitType(&FAudioVectorscope::RackUnitTypeInfo);
diff --git a/Engine/Plugins/Runtime/AudioWidgets/Source/AudioWidgets/Private/LoudnessMeterWidgetView.cpp b/Engine/Plugins/Runtime/AudioWidgets/Source/AudioWidgets/Private/LoudnessMeterWidgetView.cpp
new file mode 100644
--- /dev/null
+++ b/Engine/Plugins/Runtime/AudioWidgets/Source/AudioWidgets/Private/LoudnessMeterWidgetView.cpp
@@ -0,0 +1,255 @@
+// Copyright Epic Games, Inc. All Rights Reserved.
+
+#include "LoudnessMeterWidgetView.h"
+
+#include "AudioMeterWidgetStyle.h"
+#include "Brushes/SlateColorBrush.h"
+#include "Brushes/SlateRoundedBoxBrush.h"
+#include "SAudioMeterWidget.h"
+#include "Styling/StyleColors.h"
+#include "Widgets/Layout/SSeparator.h"
+#include "Widgets/Layout/SSpacer.h"
+#include "Widgets/Views/STileView.h"
+
+#define LOCTEXT_NAMESPACE "FLoudnessMeterWidgetView"
+
+namespace AudioWidgets
+{
+	namespace LoudnessMeterWidgetView_Private
+	{
+		FText FormatNumericValueText(TOptional<float> Value)
+		{
+			if (!Value.IsSet())
+			{
+				return FText::GetEmpty();
+			}
+			else if (*Value == -INFINITY)
+			{
+				return LOCTEXT("NegativeInfinity", "-∞");
+			}
+
+			static const FNumberFormattingOptions NumberFormattingOptions = FNumberFormattingOptions()
+				.SetMinimumFractionalDigits(1)
+				.SetMaximumFractionalDigits(1);
+
+			return FText::AsNumber(*Value, &NumberFormattingOptions);
+		}
+
+		TSharedRef<ITableRow> MakeNumericValueTileWidget(TSharedRef<const FLoudnessMeterWidgetView::FLoudnessMetric> InItem, const TSharedRef<STableViewBase>& OwnerTable)
+		{
+			using FLoudnessMetricRef = TSharedRef<const FLoudnessMeterWidgetView::FLoudnessMetric>;
+
+			return SNew(STableRow<FLoudnessMetricRef>, OwnerTable)
+				[
+					SNew(SVerticalBox)
+						+ SVerticalBox::Slot()
+						.MinHeight(8)
+						.MaxHeight(8)
+						[
+							SNew(SSpacer)
+						]
+						+ SVerticalBox::Slot()
+						.MinHeight(20)
+						.MaxHeight(20)
+						.HAlign(EHorizontalAlignment::HAlign_Center)
+						[
+							SNew(STextBlock)
+								.Text_Lambda([InItem]() { return FormatNumericValueText(InItem->Value.Get()); })
+								.Font(FStyleDefaults::GetFontInfo(14))
+								.ColorAndOpacity(FStyleColors::AccentGreen)
+								.SimpleTextMode(true)
+						]
+						+ SVerticalBox::Slot()
+						.MinHeight(32)
+						.MaxHeight(32)
+						.HAlign(EHorizontalAlignment::HAlign_Center)
+						.VAlign(EVerticalAlignment::VAlign_Center)
+						[
+							SNew(STextBlock)
+								.Text_Lambda([InItem]() { return InItem->DisplayName; })
+								.Font(FStyleDefaults::GetFontInfo(8))
+								.AutoWrapText(true)
+								.Justification(ETextJustify::Center)
+						]
+				];
+		}
+
+		TSharedRef<ITableRow> MakeMeterValueTileWidget(TSharedRef<const FLoudnessMeterWidgetView::FLoudnessMetric> InItem, const TSharedRef<STableViewBase>& OwnerTable)
+		{
+			using FLoudnessMetricRef = TSharedRef<const FLoudnessMeterWidgetView::FLoudnessMetric>;
+
+			const auto MakeAudioMeterWidgetStyle = []()
+				{
+					FAudioMeterWidgetStyle MeterStyle = FAudioMeterWidgetStyle::GetDefault();
+					MeterStyle.SetMeterSize({ 700.0, 25.0 }); // MaxLength, Width.
+					MeterStyle.SetValueRangeDb({ -60, 0 });
+					MeterStyle.SetScaleHashHeight(4.0f); // This is the tick mark length.
+					return MeterStyle;
+				};
+
+			static const FAudioMeterWidgetStyle AudioMeterWidgetStyle = MakeAudioMeterWidgetStyle();
+
+			static const FSlateRoundedBoxBrush RoundedPanelBrush(FStyleColors::Panel, 4.0f);
+
+			const FAudioMeterDefaultColorWidgetStyle& DefaultColorStyle = FAudioMeterDefaultColorWidgetStyle::GetDefault();
+
+			return SNew(STableRow<FLoudnessMetricRef>, OwnerTable)
+				[
+					SNew(SBox)
+						.VAlign(EVerticalAlignment::VAlign_Center)
+						.Padding(4.0f)
+						[
+							SNew(SBorder)
+								.BorderImage(&RoundedPanelBrush)
+								.Padding(0.0f, 8.0f)
+								[
+									SNew(SVerticalBox)
+										+ SVerticalBox::Slot()
+										.AutoHeight()
+										.HAlign(EHorizontalAlignment::HAlign_Center)
+										[
+											SNew(STextBlock)
+												.Text_Lambda([InItem]() { return FormatNumericValueText(InItem->Value.Get()); })
+												.Font(FStyleDefaults::GetFontInfo(14))
+												.ColorAndOpacity(FStyleColors::AccentGreen)
+												.SimpleTextMode(true)
+										]
+										+ SVerticalBox::Slot()
+										.HAlign(EHorizontalAlignment::HAlign_Center)
+										[
+											SNew(SAudioMeterWidget)
+												.Orientation(EOrientation::Orient_Vertical)
+												.BackgroundColor(FLinearColor::Transparent)
+												.MeterBackgroundColor(DefaultColorStyle.MeterBackgroundColor)
+												.MeterValueColor(DefaultColorStyle.MeterValueColor)
+												.MeterPeakColor(FLinearColor::Transparent)
+												.MeterClippingColor(FLinearColor::Transparent)
+												.MeterScaleColor(FLinearColor::White.CopyWithNewOpacity(0.25f))
+												.MeterScaleLabelColor(FLinearColor::White.CopyWithNewOpacity(0.25f))
+												.Style(&AudioMeterWidgetStyle)
+												.MeterChannelInfo_Lambda([InItem]() -> TArray<FAudioMeterChannelInfo>
+													{
+														constexpr float DefaultInvalidValue = -160.0f;
+														const TOptional<float> LoudnessMetricValue = InItem->Value.Get();
+														return { { LoudnessMetricValue.Get(DefaultInvalidValue), DefaultInvalidValue, DefaultInvalidValue } };
+													})
+										]
+										+ SVerticalBox::Slot()
+										.AutoHeight()
+										.MinHeight(28)
+										.HAlign(EHorizontalAlignment::HAlign_Center)
+										.VAlign(EVerticalAlignment::VAlign_Center)
+										[
+											SNew(STextBlock)
+												.Text_Lambda([InItem]() { return InItem->DisplayName; })
+												.Font(FStyleDefaults::GetFontInfo(8))
+												.AutoWrapText(true)
+												.Justification(ETextJustify::Center)
+										]
+								]
+						]
+				];
+		}
+	} // namespace LoudnessMeterWidgetView_Private
+
+	FLoudnessMeterWidgetView::FLoudnessMeterWidgetView()
+		: NumericValues(MakeShared<FLoudnessMetricRefArray>())
+		, MeterValues(MakeShared<FLoudnessMetricRefArray>())
+	{
+	}
+
+	FLoudnessMeterWidgetView::~FLoudnessMeterWidgetView()
+	{
+		NumericValues->Reset();
+		MeterValues->Reset();
+	}
+
+	TSharedRef<SWidget> FLoudnessMeterWidgetView::MakeWidget() const
+	{
+		using SLoudnessMetricTileView = STileView<FLoudnessMetricRef>;
+
+		static FSlateColorBrush BackgroundBrush(FStyleColors::Recessed);
+
+		constexpr float MeterWidth = 79.0f;
+		constexpr float MeterMinHeight = 160.0f;
+
+		TSharedRef<SLoudnessMetricTileView> MeterTileView = SNew(SLoudnessMetricTileView)
+			.OnGenerateTile_Static(&LoudnessMeterWidgetView_Private::MakeMeterValueTileWidget)
+			.ListItemsSource(MeterValues)
+			.ItemWidth(MeterWidth)
+			.ItemAlignment(EListItemAlignment::CenterAligned)
+			.SelectionMode(ESelectionMode::None)
+			.ScrollbarVisibility(EVisibility::Collapsed);
+
+		MeterTileView->SetItemHeight(TAttribute<float>::CreateLambda([&TileView = *MeterTileView]()
+			{
+				const int32 NumItems = TileView.GetItems().Num();
+				const FVector2f PanelSize = TileView.GetTickSpaceGeometry().GetLocalSize();
+				const int32 NumItemsPerLine = FMath::Max(1, FMath::FloorToInt(PanelSize.X / MeterWidth));
+				const int32 NumLines = (NumItems + NumItemsPerLine - 1) / NumItemsPerLine;
+				return FMath::Max(PanelSize.Y / NumLines, MeterMinHeight);
+			}));
+
+		return SNew(SBorder)
+			.Padding(0.0f)
+			.BorderImage(&BackgroundBrush)
+			[
+				SNew(SVerticalBox)
+					+ SVerticalBox::Slot()
+					.AutoHeight()
+					[
+						SNew(SLoudnessMetricTileView)
+							.OnGenerateTile_Static(&LoudnessMeterWidgetView_Private::MakeNumericValueTileWidget)
+							.ListItemsSource(NumericValues)
+							.ItemHeight(60.0f)
+							.ItemWidth(60.0f)
+							.ItemAlignment(EListItemAlignment::CenterAligned)
+					]
+					+ SVerticalBox::Slot()
+					.AutoHeight()
+					.Padding(8.0f, 2.0f)
+					[
+						SNew(SSeparator)
+							.SeparatorImage(FAppStyle::Get().GetBrush("WhiteBrush"))
+							.Orientation(Orient_Horizontal)
+							.Thickness(1.0f)
+							.ColorAndOpacity(FLinearColor::White.CopyWithNewOpacity(0.25f))
+					]
+					+ SVerticalBox::Slot()
+					.FillHeight(1.0f)
+					[
+						MeterTileView
+					]
+			];
+	}
+
+	void FLoudnessMeterWidgetView::AddNumericValue(const FLoudnessMetric& LoudnessMetric)
+	{
+		NumericValues->Add(MakeShared<FLoudnessMetric>(LoudnessMetric));
+	}
+
+	void FLoudnessMeterWidgetView::RemoveNumericValue(FName LoudnessMetricName)
+	{
+		const auto IsMatchingItem = [&](const FLoudnessMetricRef& LoudnessMetric) { return LoudnessMetric->Name == LoudnessMetricName; };
+		while (const FLoudnessMetricRef* MatchingItem = NumericValues->FindByPredicate(IsMatchingItem))
+		{
+			NumericValues->RemoveSingle(*MatchingItem);
+		}
+	}
+
+	void FLoudnessMeterWidgetView::AddMeterValue(const FLoudnessMetric& LoudnessMetric)
+	{
+		MeterValues->Add(MakeShared<FLoudnessMetric>(LoudnessMetric));
+	}
+
+	void FLoudnessMeterWidgetView::RemoveMeterValue(FName LoudnessMetricName)
+	{
+		const auto IsMatchingItem = [&](const FLoudnessMetricRef& LoudnessMetric) { return LoudnessMetric->Name == LoudnessMetricName; };
+		while (const FLoudnessMetricRef* MatchingItem = MeterValues->FindByPredicate(IsMatchingItem))
+		{
+			MeterValues->RemoveSingle(*MatchingItem);
+		}
+	}
+} // namespace AudioWidgets
+
+#undef LOCTEXT_NAMESPACE
diff --git a/Engine/Plugins/Runtime/AudioWidgets/Source/AudioWidgets/Private/LoudnessMeterWidgetView.h b/Engine/Plugins/Runtime/AudioWidgets/Source/AudioWidgets/Private/LoudnessMeterWidgetView.h
new file mode 100644
--- /dev/null
+++ b/Engine/Plugins/Runtime/AudioWidgets/Source/AudioWidgets/Private/LoudnessMeterWidgetView.h
@@ -0,0 +1,42 @@
+// Copyright Epic Games, Inc. All Rights Reserved.
+
+#pragma once
+
+#include "Containers/ObservableArray.h"
+#include "Widgets/SWidget.h"
+
+namespace AudioWidgets
+{
+	/**
+	 * A widget that displays various loudness metrics as numeric values and/or meters.
+	 * Could be moved to AudioWidgetsCore Public once API is stable.
+	 */
+	class FLoudnessMeterWidgetView
+	{
+	public:
+		struct FLoudnessMetric
+		{
+			FName Name;
+			FText DisplayName;
+			TAttribute<TOptional<float>> Value;
+		};
+
+		FLoudnessMeterWidgetView();
+		~FLoudnessMeterWidgetView();
+
+		TSharedRef<SWidget> MakeWidget() const;
+
+		void AddNumericValue(const FLoudnessMetric& LoudnessMetric);
+		void RemoveNumericValue(FName LoudnessMetricName);
+
+		void AddMeterValue(const FLoudnessMetric& LoudnessMetric);
+		void RemoveMeterValue(FName LoudnessMetricName);
+
+	private:
+		using FLoudnessMetricRef = TSharedRef<const FLoudnessMetric>;
+		using FLoudnessMetricRefArray = UE::Slate::Containers::TObservableArray<FLoudnessMetricRef>;
+
+		TSharedRef<FLoudnessMetricRefArray> NumericValues;
+		TSharedRef<FLoudnessMetricRefArray> MeterValues;
+	};
+} // namespace AudioWidgets
