[AudioWidgets] Fixed -INIFINITY issue with MSVC

diff --git a/Engine/Plugins/Runtime/AudioWidgets/Source/AudioWidgets/Private/LoudnessMeterWidgetView.cpp b/Engine/Plugins/Runtime/AudioWidgets/Source/AudioWidgets/Private/LoudnessMeterWidgetView.cpp
--- a/Engine/Plugins/Runtime/AudioWidgets/Source/AudioWidgets/Private/LoudnessMeterWidgetView.cpp
+++ b/Engine/Plugins/Runtime/AudioWidgets/Source/AudioWidgets/Private/LoudnessMeterWidgetView.cpp
@@ -11,6 +11,8 @@
 #include "Widgets/Layout/SSpacer.h"
 #include "Widgets/Views/STileView.h"
 
+#include <limits>
+
 #define LOCTEXT_NAMESPACE "FLoudnessMeterWidgetView"
 
 namespace AudioWidgets
@@ -23,7 +25,7 @@ namespace AudioWidgets
 			{
 				return FText::GetEmpty();
 			}
-			else if (*Value == -INFINITY)
+			else if (*Value == -std::numeric_limits<float>::infinity())
 			{
 				return LOCTEXT("NegativeInfinity", "-∞");
 			}
