[Audio Insights] Confirm calibration of loudness meter - Changed the loudness metric that is labelled in the UI as "Integrated" to be the GatedLoudness value.
- From EBU TECH 3341: "The 'Integrated Loudness' shall be measured using the gating function specified in ITU-R BS.1770".
- From Youlean Loudness Meter 2 PRO: "INTEGRATED - This readout displays the relative gated loudness averaged over the entire measurement".
- https://www.izotope.com/en/learn/what-are-lufs#integrated - "One final detail of integrated loudness is that it uses two gates".

diff --git a/Engine/Plugins/Runtime/AudioWidgets/Source/AudioWidgets/Private/AudioLoudnessMeter.cpp b/Engine/Plugins/Runtime/AudioWidgets/Source/AudioWidgets/Private/AudioLoudnessMeter.cpp
--- a/Engine/Plugins/Runtime/AudioWidgets/Source/AudioWidgets/Private/AudioLoudnessMeter.cpp
+++ b/Engine/Plugins/Runtime/AudioWidgets/Source/AudioWidgets/Private/AudioLoudnessMeter.cpp
@@ -20,9 +20,9 @@ namespace AudioWidgets
 	{
 		LoudnessMeterWidgetView.AddNumericValue(
 			{
-				.Name = GET_MEMBER_NAME_CHECKED(FLKFSResults, IntegratedLoudness),
+				.Name = GET_MEMBER_NAME_CHECKED(FLKFSResults, GatedLoudness),
 				.DisplayName = LOCTEXT("IntegratedLoudnessDisplayName", "Integrated"),
-				.Value = TAttribute<TOptional<float>>::CreateLambda([this]() { return (LatestLoudnessResults.IsSet()) ? TOptional<float>(LatestLoudnessResults->IntegratedLoudness) : NullOpt; })
+				.Value = TAttribute<TOptional<float>>::CreateLambda([this]() { return (LatestLoudnessResults.IsSet()) ? TOptional<float>(LatestLoudnessResults->GatedLoudness) : NullOpt; })
 			});
 
 		LoudnessMeterWidgetView.AddMeterValue(
