[Audio Widgets] Loudness Meter - Added analysis timespan display and reset button.

diff --git a/Engine/Plugins/Runtime/AudioWidgets/Source/AudioWidgets/Private/AudioLoudnessMeter.cpp b/Engine/Plugins/Runtime/AudioWidgets/Source/AudioWidgets/Private/AudioLoudnessMeter.cpp
--- a/Engine/Plugins/Runtime/AudioWidgets/Source/AudioWidgets/Private/AudioLoudnessMeter.cpp
+++ b/Engine/Plugins/Runtime/AudioWidgets/Source/AudioWidgets/Private/AudioLoudnessMeter.cpp
@@ -70,7 +70,16 @@ namespace AudioWidgets
 			.ExternalAudioBus = Params.AudioBusInfo.AudioBus,
 		};
 
-		return MakeShared<FAudioLoudnessMeter>(AnalyzerParams);
+		TSharedRef<FAudioLoudnessMeter> AudioLoudnessMeter = MakeShared<FAudioLoudnessMeter>(AnalyzerParams);
+
+		// Initializing these callbacks outside of the FAudioLoudnessMeter constructor allows us to use the CreateSP(...) smart pointer variants for attributes/delegates:
+		AudioLoudnessMeter->LoudnessMeterWidgetView.InitTimerPanel(
+			{
+				.AnalysisTime = TAttribute<FTimespan>::CreateSP(&AudioLoudnessMeter.Get(), &FAudioLoudnessMeter::GetAnalysisTime),
+				.OnResetButtonClicked = FOnClicked::CreateSP(AudioLoudnessMeter, &FAudioLoudnessMeter::HandleResetButtonClicked)
+			});
+
+		return AudioLoudnessMeter;
 	}
 
 	void FAudioLoudnessMeter::Init(const int32 InNumChannels, const Audio::FDeviceId InAudioDeviceId, const TObjectPtr<UAudioBus> InExternalAudioBus)
@@ -79,6 +88,8 @@ namespace AudioWidgets
 
 		LKFSSettings = TStrongObjectPtr(NewObject<ULKFSSettings>());
 
+		AudioDeviceId = InAudioDeviceId;
+
 		// Only create analyzers etc if we have an audio device:
 		if (InAudioDeviceId != FAudioBusInfo::InvalidAudioDeviceId)
 		{
@@ -135,6 +146,33 @@ namespace AudioWidgets
 			LatestLoudnessResults = InLoudnessResults;
 		}
 	}
+
+	FTimespan FAudioLoudnessMeter::GetAnalysisTime() const
+	{
+		if (LatestLoudnessResults.IsSet())
+		{
+			return FTimespan::FromSeconds(LatestLoudnessResults->Timestamp);
+		}
+	
+		return FTimespan::Zero();
+	}
+
+	FReply FAudioLoudnessMeter::HandleResetButtonClicked()
+	{
+		if (LKFSAnalyzer.IsValid())
+		{
+			LKFSAnalyzer->StopAnalyzing();
+			ReleaseLKFSAnalyzer();
+		}
+
+		if (AudioDeviceId != FAudioBusInfo::InvalidAudioDeviceId)
+		{
+			CreateLKFSAnalyzer();
+			LKFSAnalyzer->StartAnalyzing(AudioDeviceId, AudioBus.Get());
+		}
+
+		return FReply::Handled();
+	}
 } // namespace AudioWidgets
 
 #undef LOCTEXT_NAMESPACE
diff --git a/Engine/Plugins/Runtime/AudioWidgets/Source/AudioWidgets/Private/AudioLoudnessMeter.h b/Engine/Plugins/Runtime/AudioWidgets/Source/AudioWidgets/Private/AudioLoudnessMeter.h
--- a/Engine/Plugins/Runtime/AudioWidgets/Source/AudioWidgets/Private/AudioLoudnessMeter.h
+++ b/Engine/Plugins/Runtime/AudioWidgets/Source/AudioWidgets/Private/AudioLoudnessMeter.h
@@ -45,6 +45,9 @@ namespace AudioWidgets
 
 		void OnLoudnessOutput(ULKFSAnalyzer* InLKFSAnalyzer, const FLKFSResults& InLoudnessResults);
 
+		FTimespan GetAnalysisTime() const;
+		FReply HandleResetButtonClicked();
+
 		/** Creates the meters widget */
 		FLoudnessMeterWidgetView LoudnessMeterWidgetView;
 		
@@ -60,6 +63,8 @@ namespace AudioWidgets
 		/** Most recent results from the analyzer. */
 		TOptional<FLKFSResults> LatestLoudnessResults;
 
+		Audio::FDeviceId AudioDeviceId = INDEX_NONE;
+
 		/** The audio bus used for analysis. */
 		TStrongObjectPtr<UAudioBus> AudioBus;
 		bool bUseExternalAudioBus = false;
diff --git a/Engine/Plugins/Runtime/AudioWidgets/Source/AudioWidgets/Private/AudioWidgetsStyle.cpp b/Engine/Plugins/Runtime/AudioWidgets/Source/AudioWidgets/Private/AudioWidgetsStyle.cpp
--- a/Engine/Plugins/Runtime/AudioWidgets/Source/AudioWidgets/Private/AudioWidgetsStyle.cpp
+++ b/Engine/Plugins/Runtime/AudioWidgets/Source/AudioWidgets/Private/AudioWidgetsStyle.cpp
@@ -37,6 +37,7 @@ FAudioWidgetsStyle::FAudioWidgetsStyle()
 {
 	SetParentStyleName(FAppStyle::GetAppStyleSetName());
 	SetContentRoot(FPaths::EnginePluginsDir() / TEXT("Runtime/AudioWidgets/Content"));
+	SetCoreContentRoot(FPaths::EngineContentDir() / TEXT("Editor/Slate"));
 	SetResources();
 	SetupStyles();
 
@@ -208,6 +209,8 @@ void FAudioWidgetsStyle::SetupStyles()
 
 	Set("AudioWidgetsStyle.MarkerAdd", new IMAGE_BRUSH_SVG(TEXT("Slate/Icons/Marker_20"), Icon20x20));
 	Set("AudioWidgetsStyle.MarkerDelete", new IMAGE_BRUSH_SVG(TEXT("Slate/Icons/MarkerDelete_20"), Icon20x20));
+
+	Set("AudioWidgetsStyle.Reset", new CORE_IMAGE_BRUSH_SVG("Starship/Common/Reset", Icon16x16));
 }
 
 FAudioWidgetsStyle::~FAudioWidgetsStyle()
diff --git a/Engine/Plugins/Runtime/AudioWidgets/Source/AudioWidgets/Private/LoudnessMeterWidgetView.cpp b/Engine/Plugins/Runtime/AudioWidgets/Source/AudioWidgets/Private/LoudnessMeterWidgetView.cpp
--- a/Engine/Plugins/Runtime/AudioWidgets/Source/AudioWidgets/Private/LoudnessMeterWidgetView.cpp
+++ b/Engine/Plugins/Runtime/AudioWidgets/Source/AudioWidgets/Private/LoudnessMeterWidgetView.cpp
@@ -2,11 +2,13 @@
 
 #include "LoudnessMeterWidgetView.h"
 
+#include "AudioWidgetsStyle.h"
 #include "AudioMeterWidgetStyle.h"
 #include "Brushes/SlateColorBrush.h"
 #include "Brushes/SlateRoundedBoxBrush.h"
 #include "SAudioMeterWidget.h"
 #include "Styling/StyleColors.h"
+#include "Widgets/Input/SButton.h"
 #include "Widgets/Layout/SSeparator.h"
 #include "Widgets/Layout/SSpacer.h"
 #include "Widgets/Views/STileView.h"
@@ -166,6 +168,11 @@ namespace AudioWidgets
 		MeterValues->Reset();
 	}
 
+	void FLoudnessMeterWidgetView::InitTimerPanel(const FTimerPanelParams& InTimerPanelParams)
+	{
+		TimerPanelParams = InTimerPanelParams;
+	}
+
 	TSharedRef<SWidget> FLoudnessMeterWidgetView::MakeWidget() const
 	{
 		using SLoudnessMetricTileView = STileView<FLoudnessMetricRef>;
@@ -197,6 +204,50 @@ namespace AudioWidgets
 			.BorderImage(&BackgroundBrush)
 			[
 				SNew(SVerticalBox)
+					+ SVerticalBox::Slot()
+					.HAlign(EHorizontalAlignment::HAlign_Center)
+					.AutoHeight()
+					[
+						SNew(SBox)
+							.Padding(0.0f, 10.0f)
+							[
+								SNew(SHorizontalBox)
+									+ SHorizontalBox::Slot()
+									.AutoWidth()
+									[
+										SNew(SBorder)
+											.Padding(8.0f, 4.0f)
+											[
+												SNew(STextBlock)
+													.Font(FStyleDefaults::GetFontInfo(10))
+													.Text_Lambda([AnalysisTime = TimerPanelParams.AnalysisTime]() { return FText::AsTimespan(AnalysisTime.Get()); })
+											]
+									]
+									+ SHorizontalBox::Slot()
+									.AutoWidth()
+									[
+										SNew(SButton)
+											.ButtonStyle(&FAppStyle::Get().GetWidgetStyle<FButtonStyle>("SimpleButton"))
+											.OnClicked(TimerPanelParams.OnResetButtonClicked)
+											[
+												SNew(SImage)
+													.Image(FAudioWidgetsStyle::Get().GetBrush("AudioWidgetsStyle.Reset"))
+													.ColorAndOpacity(FSlateColor::UseForeground())
+													.ToolTipText(LOCTEXT("ResetButtonToolTip", "Reset Integrated loudness analysis"))
+											]
+									]
+							]
+					]
+					+ SVerticalBox::Slot()
+					.AutoHeight()
+					.Padding(8.0f, 2.0f)
+					[
+						SNew(SSeparator)
+							.SeparatorImage(FAppStyle::Get().GetBrush("WhiteBrush"))
+							.Orientation(Orient_Horizontal)
+							.Thickness(1.0f)
+							.ColorAndOpacity(FLinearColor::White.CopyWithNewOpacity(0.25f))
+					]
 					+ SVerticalBox::Slot()
 					.AutoHeight()
 					[
diff --git a/Engine/Plugins/Runtime/AudioWidgets/Source/AudioWidgets/Private/LoudnessMeterWidgetView.h b/Engine/Plugins/Runtime/AudioWidgets/Source/AudioWidgets/Private/LoudnessMeterWidgetView.h
--- a/Engine/Plugins/Runtime/AudioWidgets/Source/AudioWidgets/Private/LoudnessMeterWidgetView.h
+++ b/Engine/Plugins/Runtime/AudioWidgets/Source/AudioWidgets/Private/LoudnessMeterWidgetView.h
@@ -3,6 +3,7 @@
 #pragma once
 
 #include "Containers/ObservableArray.h"
+#include "Framework/SlateDelegates.h"
 #include "Widgets/SWidget.h"
 
 namespace AudioWidgets
@@ -21,9 +22,17 @@ namespace AudioWidgets
 			TAttribute<TOptional<float>> Value;
 		};
 
+		struct FTimerPanelParams
+		{
+			TAttribute<FTimespan> AnalysisTime = FTimespan::Zero();
+			FOnClicked OnResetButtonClicked;
+		};
+
 		FLoudnessMeterWidgetView();
 		~FLoudnessMeterWidgetView();
 
+		void InitTimerPanel(const FTimerPanelParams& InTimerPanelParams);
+
 		TSharedRef<SWidget> MakeWidget() const;
 
 		void AddNumericValue(const FLoudnessMetric& LoudnessMetric);
@@ -38,5 +47,6 @@ namespace AudioWidgets
 
 		TSharedRef<FLoudnessMetricRefArray> NumericValues;
 		TSharedRef<FLoudnessMetricRefArray> MeterValues;
+		FTimerPanelParams TimerPanelParams;
 	};
 } // namespace AudioWidgets
