[Audio Widgets] Loudness Meter - Implement BugHawk suggestion of also clearing LatestLoudnessResults when reset button clicked.

diff --git a/Engine/Plugins/Runtime/AudioWidgets/Source/AudioWidgets/Private/AudioLoudnessMeter.cpp b/Engine/Plugins/Runtime/AudioWidgets/Source/AudioWidgets/Private/AudioLoudnessMeter.cpp
--- a/Engine/Plugins/Runtime/AudioWidgets/Source/AudioWidgets/Private/AudioLoudnessMeter.cpp
+++ b/Engine/Plugins/Runtime/AudioWidgets/Source/AudioWidgets/Private/AudioLoudnessMeter.cpp
@@ -165,6 +165,8 @@ namespace AudioWidgets
 			ReleaseLKFSAnalyzer();
 		}
 
+		LatestLoudnessResults.Reset();
+
 		if (AudioDeviceId != FAudioBusInfo::InvalidAudioDeviceId)
 		{
 			CreateLKFSAnalyzer();
