[Audio Widgets] Loudness Meter - Added toolbar and options menu, includes Show/Hide Timer option.

diff --git a/Engine/Plugins/AudioInsights/Source/AudioInsightsEditor/Public/AudioInsightsEditorSettings.h b/Engine/Plugins/AudioInsights/Source/AudioInsightsEditor/Public/AudioInsightsEditorSettings.h
--- a/Engine/Plugins/AudioInsights/Source/AudioInsightsEditor/Public/AudioInsightsEditorSettings.h
+++ b/Engine/Plugins/AudioInsights/Source/AudioInsightsEditor/Public/AudioInsightsEditorSettings.h
@@ -5,6 +5,7 @@
 #include "AudioSpectrogram.h"
 #include "AudioSpectrumAnalyzer.h"
 #include "Engine/DeveloperSettings.h"
+#include "LoudnessMeterRackUnitSettings.h"
 #include "Settings/AudioEventLogSettings.h"
 #include "Settings/SoundDashboardSettings.h"
 
@@ -41,6 +42,10 @@ public:
 	UPROPERTY(EditAnywhere, config, Category = SpectrumAnalyzer, meta = (ShowOnlyInnerProperties))
 	FSpectrumAnalyzerRackUnitSettings SpectrumAnalyzerSettings;
 
+	/** Settings for analyzer rack loudness meter widget */
+	UPROPERTY(EditAnywhere, config, Category = LoudnessMeter, meta = (ShowOnlyInnerProperties))
+	FLoudnessMeterRackUnitSettings LoudnessMeterSettings;
+
 	/** Settings for the Event Log */
 	UPROPERTY(EditAnywhere, config, Category = EventLog, meta = (ShowOnlyInnerProperties))
 	FAudioEventLogSettings EventLogSettings;
diff --git a/Engine/Plugins/Runtime/AudioWidgets/Source/AudioWidgets/Private/AudioLoudnessMeter.cpp b/Engine/Plugins/Runtime/AudioWidgets/Source/AudioWidgets/Private/AudioLoudnessMeter.cpp
--- a/Engine/Plugins/Runtime/AudioWidgets/Source/AudioWidgets/Private/AudioLoudnessMeter.cpp
+++ b/Engine/Plugins/Runtime/AudioWidgets/Source/AudioWidgets/Private/AudioLoudnessMeter.cpp
@@ -2,12 +2,45 @@
 
 #include "AudioLoudnessMeter.h"
 
+#include "LoudnessMeterRackUnitSettings.h"
 #include "Widgets/Docking/SDockTab.h"
 
 #define LOCTEXT_NAMESPACE "FAudioLoudnessMeter"
 
 namespace AudioWidgets
 {
+	/**
+	 * Light wrapper for accessing settings for analyzer rack units. Can be passed by value.
+	 */
+	template<class SettingsT>
+	class TRackUnitSettingsHelper
+	{
+	public:
+		TRackUnitSettingsHelper(const FProperty& InSettingsProperty)
+			: SettingsProperty(InSettingsProperty)
+		{
+		}
+
+		SettingsT& GetRackUnitSettings() const
+		{
+			UObject* EditorSettingsObject = GetEditorSettingsObject();
+			return *(SettingsProperty.ContainerPtrToValuePtr<SettingsT>(EditorSettingsObject));
+		}
+
+		void SaveConfig() const
+		{
+			GetEditorSettingsObject()->SaveConfig();
+		}
+
+	private:
+		UObject* GetEditorSettingsObject() const
+		{
+			return SettingsProperty.GetOwnerClass()->GetDefaultObject();
+		}
+
+		const FProperty& SettingsProperty;
+	};
+
 	const FAudioAnalyzerRackUnitTypeInfo FAudioLoudnessMeter::RackUnitTypeInfo
 	{
 		.TypeName = TEXT("FAudioLoudnessMeter"),
@@ -71,13 +104,7 @@ namespace AudioWidgets
 		};
 
 		TSharedRef<FAudioLoudnessMeter> AudioLoudnessMeter = MakeShared<FAudioLoudnessMeter>(AnalyzerParams);
-
-		// Initializing these callbacks outside of the FAudioLoudnessMeter constructor allows us to use the CreateSP(...) smart pointer variants for attributes/delegates:
-		AudioLoudnessMeter->LoudnessMeterWidgetView.InitTimerPanel(
-			{
-				.AnalysisTime = TAttribute<FTimespan>::CreateSP(&AudioLoudnessMeter.Get(), &FAudioLoudnessMeter::GetAnalysisTime),
-				.OnResetButtonClicked = FOnClicked::CreateSP(AudioLoudnessMeter, &FAudioLoudnessMeter::HandleResetButtonClicked)
-			});
+		AudioLoudnessMeter->InitLoudnessMeterWidgetView(Params);
 
 		return AudioLoudnessMeter;
 	}
@@ -139,6 +166,37 @@ namespace AudioWidgets
 		LKFSAnalyzer.Reset();
 	}
 
+	void FAudioLoudnessMeter::InitLoudnessMeterWidgetView(const FAudioAnalyzerRackUnitConstructParams& Params)
+	{
+		// Initializing these callbacks outside of the FAudioLoudnessMeter constructor allows us to use the CreateSP(...) smart pointer variants for attributes/delegates:
+
+		FLoudnessMeterWidgetView::FTimerPanelParams TimerPanelParams
+		{
+			.AnalysisTime = TAttribute<FTimespan>::CreateSP(this, &FAudioLoudnessMeter::GetAnalysisTime),
+			.OnResetButtonClicked = FOnClicked::CreateSP(this, &FAudioLoudnessMeter::HandleResetButtonClicked),
+			.bIsVisible = true
+		};
+
+		if (Params.EditorSettingsClass != nullptr)
+		{
+			// If we have been given a valid editor settings class, bind analyzer options to the settings:
+			if (const FProperty* LoudnessMeterSettingsProperty = Params.EditorSettingsClass->FindPropertyByName("LoudnessMeterSettings"))
+			{
+				const TRackUnitSettingsHelper<FLoudnessMeterRackUnitSettings> SettingsHelper(*LoudnessMeterSettingsProperty);
+
+				TimerPanelParams.bIsVisible.BindLambda([=]() { return SettingsHelper.GetRackUnitSettings().bDisplayAnalysisTimer; });
+				TimerPanelParams.OnVisibilityToggleRequested.BindLambda([=]()
+					{
+						FLoudnessMeterRackUnitSettings& LoudnessMeterSettings = SettingsHelper.GetRackUnitSettings();
+						LoudnessMeterSettings.bDisplayAnalysisTimer = !LoudnessMeterSettings.bDisplayAnalysisTimer;
+						SettingsHelper.SaveConfig();
+					});
+			}
+		}
+
+		LoudnessMeterWidgetView.InitTimerPanel(TimerPanelParams);
+	}
+
 	void FAudioLoudnessMeter::OnLoudnessOutput(ULKFSAnalyzer* InLKFSAnalyzer, const FLKFSResults& InLoudnessResults)
 	{
 		if (InLKFSAnalyzer == LKFSAnalyzer.Get())
diff --git a/Engine/Plugins/Runtime/AudioWidgets/Source/AudioWidgets/Private/AudioLoudnessMeter.h b/Engine/Plugins/Runtime/AudioWidgets/Source/AudioWidgets/Private/AudioLoudnessMeter.h
--- a/Engine/Plugins/Runtime/AudioWidgets/Source/AudioWidgets/Private/AudioLoudnessMeter.h
+++ b/Engine/Plugins/Runtime/AudioWidgets/Source/AudioWidgets/Private/AudioLoudnessMeter.h
@@ -43,6 +43,8 @@ namespace AudioWidgets
 		void CreateLKFSAnalyzer();
 		void ReleaseLKFSAnalyzer();
 
+		void InitLoudnessMeterWidgetView(const FAudioAnalyzerRackUnitConstructParams& Params);
+
 		void OnLoudnessOutput(ULKFSAnalyzer* InLKFSAnalyzer, const FLKFSResults& InLoudnessResults);
 
 		FTimespan GetAnalysisTime() const;
diff --git a/Engine/Plugins/Runtime/AudioWidgets/Source/AudioWidgets/Private/LoudnessMeterWidgetView.cpp b/Engine/Plugins/Runtime/AudioWidgets/Source/AudioWidgets/Private/LoudnessMeterWidgetView.cpp
--- a/Engine/Plugins/Runtime/AudioWidgets/Source/AudioWidgets/Private/LoudnessMeterWidgetView.cpp
+++ b/Engine/Plugins/Runtime/AudioWidgets/Source/AudioWidgets/Private/LoudnessMeterWidgetView.cpp
@@ -6,8 +6,10 @@
 #include "AudioMeterWidgetStyle.h"
 #include "Brushes/SlateColorBrush.h"
 #include "Brushes/SlateRoundedBoxBrush.h"
+#include "Framework/MultiBox/MultiBoxBuilder.h"
 #include "SAudioMeterWidget.h"
 #include "Styling/StyleColors.h"
+#include "Widgets/Colors/SColorBlock.h"
 #include "Widgets/Input/SButton.h"
 #include "Widgets/Layout/SSeparator.h"
 #include "Widgets/Layout/SSpacer.h"
@@ -179,6 +181,48 @@ namespace AudioWidgets
 
 		static FSlateColorBrush BackgroundBrush(FStyleColors::Recessed);
 
+		constexpr bool bForceSmallIcons = true;
+		const TSharedRef<FUICommandList> CommandList = MakeShared<FUICommandList>();
+		FSlimHorizontalToolBarBuilder ToolBarBuilder(CommandList, FMultiBoxCustomization::None, nullptr, bForceSmallIcons);
+
+		const auto MakeOptionsMenu = [CommandList, TimerPanelParams = this->TimerPanelParams]() -> TSharedRef<SWidget>
+			{
+				constexpr bool bShouldCloseWindowAfterMenuSelection = true;
+				FMenuBuilder MenuBuilder(bShouldCloseWindowAfterMenuSelection, CommandList);
+
+				{
+					MenuBuilder.BeginSection(NAME_None, LOCTEXT("LayoutOptions", "Layout Options"));
+
+					if (TimerPanelParams.OnVisibilityToggleRequested.IsBound())
+					{
+						MenuBuilder.AddMenuEntry(
+							LOCTEXT("LayoutOptions_Time_Label", "Time"),
+							LOCTEXT("LayoutOptions_Time_ToolTip", "Display analysis timer"),
+							FSlateIcon(),
+							FUIAction(
+								TimerPanelParams.OnVisibilityToggleRequested,
+								FCanExecuteAction(),
+								FIsActionChecked::CreateLambda([bIsVisible = TimerPanelParams.bIsVisible]() { return bIsVisible.Get(); })
+							),
+							NAME_None,
+							EUserInterfaceActionType::ToggleButton);
+					}
+
+					MenuBuilder.EndSection();
+				}
+
+				return MenuBuilder.MakeWidget();
+			};
+
+		ToolBarBuilder.AddComboButton(
+			FUIAction(),
+			FOnGetContent::CreateLambda(MakeOptionsMenu),
+			LOCTEXT("OptionsButtonDisplayName", "Options"),
+			FText(),
+			FSlateIcon(FAppStyle::GetAppStyleSetName(), "LevelEditor.GameSettings"));
+
+		const auto GetTimerPanelVisibility = [bIsVisible = TimerPanelParams.bIsVisible]() { return (bIsVisible.Get()) ? EVisibility::Visible : EVisibility::Collapsed; };
+
 		constexpr float MeterWidth = 79.0f;
 		constexpr float MeterMinHeight = 160.0f;
 
@@ -204,12 +248,28 @@ namespace AudioWidgets
 			.BorderImage(&BackgroundBrush)
 			[
 				SNew(SVerticalBox)
+					+ SVerticalBox::Slot()
+					.AutoHeight()
+					[
+						SNew(SOverlay)
+							+ SOverlay::Slot()
+							[
+								SNew(SColorBlock)
+									.Color(FStyleColors::Panel.GetSpecifiedColor())
+							]
+							+ SOverlay::Slot()
+							.HAlign(EHorizontalAlignment::HAlign_Right)
+							[
+								ToolBarBuilder.MakeWidget()
+							]
+					]
 					+ SVerticalBox::Slot()
 					.HAlign(EHorizontalAlignment::HAlign_Center)
 					.AutoHeight()
 					[
 						SNew(SBox)
 							.Padding(0.0f, 10.0f)
+							.Visibility_Lambda(GetTimerPanelVisibility)
 							[
 								SNew(SHorizontalBox)
 									+ SHorizontalBox::Slot()
@@ -243,6 +303,7 @@ namespace AudioWidgets
 					.Padding(8.0f, 2.0f)
 					[
 						SNew(SSeparator)
+							.Visibility_Lambda(GetTimerPanelVisibility)
 							.SeparatorImage(FAppStyle::Get().GetBrush("WhiteBrush"))
 							.Orientation(Orient_Horizontal)
 							.Thickness(1.0f)
diff --git a/Engine/Plugins/Runtime/AudioWidgets/Source/AudioWidgets/Private/LoudnessMeterWidgetView.h b/Engine/Plugins/Runtime/AudioWidgets/Source/AudioWidgets/Private/LoudnessMeterWidgetView.h
--- a/Engine/Plugins/Runtime/AudioWidgets/Source/AudioWidgets/Private/LoudnessMeterWidgetView.h
+++ b/Engine/Plugins/Runtime/AudioWidgets/Source/AudioWidgets/Private/LoudnessMeterWidgetView.h
@@ -26,6 +26,8 @@ namespace AudioWidgets
 		{
 			TAttribute<FTimespan> AnalysisTime = FTimespan::Zero();
 			FOnClicked OnResetButtonClicked;
+			TAttribute<bool> bIsVisible = false;
+			FSimpleDelegate OnVisibilityToggleRequested;
 		};
 
 		FLoudnessMeterWidgetView();
diff --git a/Engine/Plugins/Runtime/AudioWidgets/Source/AudioWidgets/Public/LoudnessMeterRackUnitSettings.h b/Engine/Plugins/Runtime/AudioWidgets/Source/AudioWidgets/Public/LoudnessMeterRackUnitSettings.h
new file mode 100644
--- /dev/null
+++ b/Engine/Plugins/Runtime/AudioWidgets/Source/AudioWidgets/Public/LoudnessMeterRackUnitSettings.h
@@ -0,0 +1,13 @@
+// Copyright Epic Games, Inc. All Rights Reserved.
+#pragma once
+
+#include "LoudnessMeterRackUnitSettings.generated.h"
+
+USTRUCT()
+struct FLoudnessMeterRackUnitSettings
+{
+	GENERATED_BODY()
+
+	UPROPERTY(EditAnywhere, config, Category = LoudnessMeter)
+	bool bDisplayAnalysisTimer = true;
+};
diff --git a/Engine/Plugins/Runtime/Metasound/Source/MetasoundEditor/Private/MetasoundEditorSettings.h b/Engine/Plugins/Runtime/Metasound/Source/MetasoundEditor/Private/MetasoundEditorSettings.h
--- a/Engine/Plugins/Runtime/Metasound/Source/MetasoundEditor/Private/MetasoundEditorSettings.h
+++ b/Engine/Plugins/Runtime/Metasound/Source/MetasoundEditor/Private/MetasoundEditorSettings.h
@@ -3,6 +3,7 @@
 
 #include "AudioSpectrogram.h"
 #include "AudioSpectrumAnalyzer.h"
+#include "LoudnessMeterRackUnitSettings.h"
 #include "MetasoundFrontendDocument.h"
 #include "Misc/CoreDefines.h"
 #include "ToolMenuMisc.h"
@@ -267,6 +268,10 @@ public:
 	UPROPERTY(EditAnywhere, config, Category = SpectrumAnalyzer, meta = (ShowOnlyInnerProperties))
 	FSpectrumAnalyzerRackUnitSettings SpectrumAnalyzerSettings;
 
+	/** Settings for metasound output loudness meter widget */
+	UPROPERTY(EditAnywhere, config, Category = LoudnessMeter, meta = (ShowOnlyInnerProperties))
+	FLoudnessMeterRackUnitSettings LoudnessMeterSettings;
+
 	/** Widget type to show on input nodes by default */
 	UPROPERTY(EditAnywhere, config, Category = General)
 	EMetasoundMemberDefaultWidget DefaultInputWidgetType = EMetasoundMemberDefaultWidget::RadialSlider;
