[Audio Widgets] Loudness Meter - Add layout options to show/hide the different loudness metrics.

diff --git a/Engine/Plugins/Runtime/AudioWidgets/Source/AudioWidgets/Private/AudioLoudnessMeter.cpp b/Engine/Plugins/Runtime/AudioWidgets/Source/AudioWidgets/Private/AudioLoudnessMeter.cpp
--- a/Engine/Plugins/Runtime/AudioWidgets/Source/AudioWidgets/Private/AudioLoudnessMeter.cpp
+++ b/Engine/Plugins/Runtime/AudioWidgets/Source/AudioWidgets/Private/AudioLoudnessMeter.cpp
@@ -51,31 +51,17 @@ namespace AudioWidgets
 
 	FAudioLoudnessMeter::FAudioLoudnessMeter(const FAudioLoudnessMeterParams& Params)
 	{
-		LoudnessMeterWidgetView.AddNumericValue(
-			{
-				.Name = GET_MEMBER_NAME_CHECKED(FLKFSResults, GatedLoudness),
-				.DisplayName = LOCTEXT("IntegratedLoudnessDisplayName", "Integrated"),
-				.Value = TAttribute<TOptional<float>>::CreateLambda([this]() { return (LatestLoudnessResults.IsSet()) ? TOptional<float>(LatestLoudnessResults->GatedLoudness) : NullOpt; })
-			});
-
-		LoudnessMeterWidgetView.AddMeterValue(
-			{
-				.Name = GET_MEMBER_NAME_CHECKED(FLKFSResults, ShortTermLoudness),
-				.DisplayName = LOCTEXT("ShortTermLoudnessDisplayName", "Short Term"),
-				.Value = TAttribute<TOptional<float>>::CreateLambda([this]() { return (LatestLoudnessResults.IsSet()) ? TOptional<float>(LatestLoudnessResults->ShortTermLoudness) : NullOpt; })
-			});
-		LoudnessMeterWidgetView.AddMeterValue(
-			{
-				.Name = GET_MEMBER_NAME_CHECKED(FLKFSResults, Loudness),
-				.DisplayName = LOCTEXT("MomentaryLoudnessDisplayName", "Momentary"),
-				.Value = TAttribute<TOptional<float>>::CreateLambda([this]() { return (LatestLoudnessResults.IsSet()) ? TOptional<float>(LatestLoudnessResults->Loudness) : NullOpt; })
-			});
-
 		Init(Params.NumChannels, Params.AudioDeviceId, Params.ExternalAudioBus);
 	}
 
 	FAudioLoudnessMeter::~FAudioLoudnessMeter()
 	{
+		if (WidgetRefreshTicker.IsValid())
+		{
+			FTSTicker::RemoveTicker(WidgetRefreshTicker);
+			WidgetRefreshTicker.Reset();
+		}
+
 		Teardown();
 	}
 
@@ -177,6 +163,31 @@ namespace AudioWidgets
 			.bIsVisible = true
 		};
 
+		FLoudnessMeterWidgetView::FLoudnessMetric LoudnessMetrics[] =
+		{
+			{
+				.Name = GET_MEMBER_NAME_CHECKED(FLoudnessMeterRackUnitSettings, IntegratedLoudness),
+				.DisplayName = LOCTEXT("IntegratedLoudnessDisplayName", "Integrated"),
+				.Value = TAttribute<TOptional<float>>::CreateSPLambda(this, [this]() { return (LatestLoudnessResults.IsSet()) ? TOptional<float>(LatestLoudnessResults->GatedLoudness) : NullOpt; }),
+				.bShowValue = true,
+				.bShowMeter = false,
+			},
+			{
+				.Name = GET_MEMBER_NAME_CHECKED(FLoudnessMeterRackUnitSettings, ShortTermLoudness),
+				.DisplayName = LOCTEXT("ShortTermLoudnessDisplayName", "Short Term"),
+				.Value = TAttribute<TOptional<float>>::CreateSPLambda(this, [this]() { return (LatestLoudnessResults.IsSet()) ? TOptional<float>(LatestLoudnessResults->ShortTermLoudness) : NullOpt; }),
+				.bShowValue = false,
+				.bShowMeter = true,
+			},
+			{
+				.Name = GET_MEMBER_NAME_CHECKED(FLoudnessMeterRackUnitSettings, MomentaryLoudness),
+				.DisplayName = LOCTEXT("MomentaryLoudnessDisplayName", "Momentary"),
+				.Value = TAttribute<TOptional<float>>::CreateSPLambda(this, [this]() { return (LatestLoudnessResults.IsSet()) ? TOptional<float>(LatestLoudnessResults->Loudness) : NullOpt; }),
+				.bShowValue = false,
+				.bShowMeter = true,
+			}
+		};
+
 		if (Params.EditorSettingsClass != nullptr)
 		{
 			// If we have been given a valid editor settings class, bind analyzer options to the settings:
@@ -191,10 +202,53 @@ namespace AudioWidgets
 						LoudnessMeterSettings.bDisplayAnalysisTimer = !LoudnessMeterSettings.bDisplayAnalysisTimer;
 						SettingsHelper.SaveConfig();
 					});
+
+				for (FLoudnessMeterWidgetView::FLoudnessMetric& LoudnessMetric : LoudnessMetrics)
+				{
+					if (const FProperty* DisplayOptionsProperty = FLoudnessMeterRackUnitSettings::StaticStruct()->FindPropertyByName(LoudnessMetric.Name))
+					{
+						LoudnessMetric.bShowValue.BindLambda([=]()
+							{
+								const FLoudnessMetricDisplayOptions* DisplayOptions = DisplayOptionsProperty->ContainerPtrToValuePtr<FLoudnessMetricDisplayOptions>(&SettingsHelper.GetRackUnitSettings());
+								return DisplayOptions->bShowValue;
+							});
+
+						LoudnessMetric.bShowMeter.BindLambda([=]()
+							{
+								const FLoudnessMetricDisplayOptions* DisplayOptions = DisplayOptionsProperty->ContainerPtrToValuePtr<FLoudnessMetricDisplayOptions>(&SettingsHelper.GetRackUnitSettings());
+								return DisplayOptions->bShowMeter;
+							});
+
+						LoudnessMetric.OnShowValueToggleRequested.BindLambda([=]()
+							{
+								FLoudnessMetricDisplayOptions* DisplayOptions = DisplayOptionsProperty->ContainerPtrToValuePtr<FLoudnessMetricDisplayOptions>(&SettingsHelper.GetRackUnitSettings());
+								DisplayOptions->bShowValue = !DisplayOptions->bShowValue;
+								SettingsHelper.SaveConfig();
+							});
+
+						LoudnessMetric.OnShowMeterToggleRequested.BindLambda([=]()
+							{
+								FLoudnessMetricDisplayOptions* DisplayOptions = DisplayOptionsProperty->ContainerPtrToValuePtr<FLoudnessMetricDisplayOptions>(&SettingsHelper.GetRackUnitSettings());
+								DisplayOptions->bShowMeter = !DisplayOptions->bShowMeter;
+								SettingsHelper.SaveConfig();
+							});
+					}
+				}
+
+				// We just poll for visible loudness metrics:
+				WidgetRefreshTicker = FTSTicker::GetCoreTicker().AddTicker(TEXT("FAudioLoudnessMeter::WidgetRefreshTicker"), 0.0f, [this](float DeltaTime)
+					{
+						LoudnessMeterWidgetView.RefreshVisibleLoudnessMetrics();
+						return true;
+					});
 			}
 		}
 
 		LoudnessMeterWidgetView.InitTimerPanel(TimerPanelParams);
+		for (const FLoudnessMeterWidgetView::FLoudnessMetric& LoudnessMetric : LoudnessMetrics)
+		{
+			LoudnessMeterWidgetView.AddLoudnessMetric(LoudnessMetric);
+		}
 	}
 
 	void FAudioLoudnessMeter::OnLoudnessOutput(ULKFSAnalyzer* InLKFSAnalyzer, const FLKFSResults& InLoudnessResults)
diff --git a/Engine/Plugins/Runtime/AudioWidgets/Source/AudioWidgets/Private/AudioLoudnessMeter.h b/Engine/Plugins/Runtime/AudioWidgets/Source/AudioWidgets/Private/AudioLoudnessMeter.h
--- a/Engine/Plugins/Runtime/AudioWidgets/Source/AudioWidgets/Private/AudioLoudnessMeter.h
+++ b/Engine/Plugins/Runtime/AudioWidgets/Source/AudioWidgets/Private/AudioLoudnessMeter.h
@@ -2,6 +2,7 @@
 #pragma once
 
 #include "AudioAnalyzerRack.h"
+#include "Containers/Ticker.h"
 #include "LKFS.h"
 #include "LoudnessMeterWidgetView.h"
 
@@ -67,6 +68,8 @@ namespace AudioWidgets
 
 		Audio::FDeviceId AudioDeviceId = INDEX_NONE;
 
+		FTSTicker::FDelegateHandle WidgetRefreshTicker;
+
 		/** The audio bus used for analysis. */
 		TStrongObjectPtr<UAudioBus> AudioBus;
 		bool bUseExternalAudioBus = false;
diff --git a/Engine/Plugins/Runtime/AudioWidgets/Source/AudioWidgets/Private/LoudnessMeterWidgetView.cpp b/Engine/Plugins/Runtime/AudioWidgets/Source/AudioWidgets/Private/LoudnessMeterWidgetView.cpp
--- a/Engine/Plugins/Runtime/AudioWidgets/Source/AudioWidgets/Private/LoudnessMeterWidgetView.cpp
+++ b/Engine/Plugins/Runtime/AudioWidgets/Source/AudioWidgets/Private/LoudnessMeterWidgetView.cpp
@@ -159,13 +159,15 @@ namespace AudioWidgets
 	} // namespace LoudnessMeterWidgetView_Private
 
 	FLoudnessMeterWidgetView::FLoudnessMeterWidgetView()
-		: NumericValues(MakeShared<FLoudnessMetricRefArray>())
+		: LoudnessMetrics(MakeShared<FLoudnessMetricRefArray>())
+		, NumericValues(MakeShared<FLoudnessMetricRefArray>())
 		, MeterValues(MakeShared<FLoudnessMetricRefArray>())
 	{
 	}
 
 	FLoudnessMeterWidgetView::~FLoudnessMeterWidgetView()
 	{
+		LoudnessMetrics->Reset();
 		NumericValues->Reset();
 		MeterValues->Reset();
 	}
@@ -185,7 +187,7 @@ namespace AudioWidgets
 		const TSharedRef<FUICommandList> CommandList = MakeShared<FUICommandList>();
 		FSlimHorizontalToolBarBuilder ToolBarBuilder(CommandList, FMultiBoxCustomization::None, nullptr, bForceSmallIcons);
 
-		const auto MakeOptionsMenu = [CommandList, TimerPanelParams = this->TimerPanelParams]() -> TSharedRef<SWidget>
+		const auto MakeOptionsMenu = [CommandList, TimerPanelParams = this->TimerPanelParams, LoudnessMetrics = this->LoudnessMetrics]() -> TSharedRef<SWidget>
 			{
 				constexpr bool bShouldCloseWindowAfterMenuSelection = true;
 				FMenuBuilder MenuBuilder(bShouldCloseWindowAfterMenuSelection, CommandList);
@@ -208,6 +210,41 @@ namespace AudioWidgets
 							EUserInterfaceActionType::ToggleButton);
 					}
 
+					for (const FLoudnessMetricRef& LoudnessMetricRef : LoudnessMetrics.Get())
+					{
+						MenuBuilder.AddSubMenu(
+							LoudnessMetricRef->DisplayName,
+							FText(),
+							FNewMenuDelegate::CreateSPLambda(LoudnessMetricRef, [&LoudnessMetric = LoudnessMetricRef.Get()](FMenuBuilder& SubMenu)
+								{
+									SubMenu.BeginSection(NAME_None, LoudnessMetric.DisplayName);
+
+									SubMenu.AddMenuEntry(
+										LOCTEXT("LayoutOptions_LoudnessMetric_ShowValue", "Show Value"),
+										FText(),
+										FSlateIcon(),
+										FUIAction(
+											LoudnessMetric.OnShowValueToggleRequested,
+											FCanExecuteAction::CreateLambda([bCanExecute = LoudnessMetric.OnShowValueToggleRequested.IsBound()]() { return bCanExecute; }),
+											FIsActionChecked::CreateLambda([bShowValue = LoudnessMetric.bShowValue]() { return bShowValue.Get(); })),
+										NAME_None,
+										EUserInterfaceActionType::ToggleButton);
+
+									SubMenu.AddMenuEntry(
+										LOCTEXT("LayoutOptions_LoudnessMetric_ShowMeter", "Show Meter"),
+										FText(),
+										FSlateIcon(),
+										FUIAction(
+											LoudnessMetric.OnShowMeterToggleRequested,
+											FCanExecuteAction::CreateLambda([bCanExecute = LoudnessMetric.OnShowMeterToggleRequested.IsBound()]() { return bCanExecute; }),
+											FIsActionChecked::CreateLambda([bShowMeter = LoudnessMetric.bShowMeter]() { return bShowMeter.Get(); })),
+										NAME_None,
+										EUserInterfaceActionType::ToggleButton);
+
+									SubMenu.EndSection();
+								}));
+					}
+
 					MenuBuilder.EndSection();
 				}
 
@@ -222,6 +259,10 @@ namespace AudioWidgets
 			FSlateIcon(FAppStyle::GetAppStyleSetName(), "LevelEditor.GameSettings"));
 
 		const auto GetTimerPanelVisibility = [bIsVisible = TimerPanelParams.bIsVisible]() { return (bIsVisible.Get()) ? EVisibility::Visible : EVisibility::Collapsed; };
+		const TAttribute<EVisibility> GetNumericValuesPanelVisibility = TAttribute<EVisibility>::CreateSPLambda(NumericValues, [&InNumericValues = NumericValues.Get()]()
+			{
+				return (!InNumericValues.IsEmpty()) ? EVisibility::Visible : EVisibility::Collapsed;
+			});
 
 		constexpr float MeterWidth = 79.0f;
 		constexpr float MeterMinHeight = 160.0f;
@@ -313,6 +354,7 @@ namespace AudioWidgets
 					.AutoHeight()
 					[
 						SNew(SLoudnessMetricTileView)
+							.Visibility(GetNumericValuesPanelVisibility)
 							.OnGenerateTile_Static(&LoudnessMeterWidgetView_Private::MakeNumericValueTileWidget)
 							.ListItemsSource(NumericValues)
 							.ItemHeight(60.0f)
@@ -324,6 +366,7 @@ namespace AudioWidgets
 					.Padding(8.0f, 2.0f)
 					[
 						SNew(SSeparator)
+							.Visibility(GetNumericValuesPanelVisibility)
 							.SeparatorImage(FAppStyle::Get().GetBrush("WhiteBrush"))
 							.Orientation(Orient_Horizontal)
 							.Thickness(1.0f)
@@ -337,31 +380,53 @@ namespace AudioWidgets
 			];
 	}
 
-	void FLoudnessMeterWidgetView::AddNumericValue(const FLoudnessMetric& LoudnessMetric)
+	void FLoudnessMeterWidgetView::AddLoudnessMetric(const FLoudnessMetric& LoudnessMetric)
 	{
-		NumericValues->Add(MakeShared<FLoudnessMetric>(LoudnessMetric));
-	}
+		FLoudnessMetricRef LoudnessMetricRef = MakeShared<FLoudnessMetric>(LoudnessMetric);
+		LoudnessMetrics->Add(LoudnessMetricRef);
 
-	void FLoudnessMeterWidgetView::RemoveNumericValue(FName LoudnessMetricName)
-	{
-		const auto IsMatchingItem = [&](const FLoudnessMetricRef& LoudnessMetric) { return LoudnessMetric->Name == LoudnessMetricName; };
-		while (const FLoudnessMetricRef* MatchingItem = NumericValues->FindByPredicate(IsMatchingItem))
+		if (LoudnessMetric.bShowValue.Get())
 		{
-			NumericValues->RemoveSingle(*MatchingItem);
+			NumericValues->Add(LoudnessMetricRef);
 		}
-	}
 
-	void FLoudnessMeterWidgetView::AddMeterValue(const FLoudnessMetric& LoudnessMetric)
-	{
-		MeterValues->Add(MakeShared<FLoudnessMetric>(LoudnessMetric));
+		if (LoudnessMetric.bShowMeter.Get())
+		{
+			MeterValues->Add(LoudnessMetricRef);
+		}
 	}
 
-	void FLoudnessMeterWidgetView::RemoveMeterValue(FName LoudnessMetricName)
+	void FLoudnessMeterWidgetView::RefreshVisibleLoudnessMetrics()
 	{
-		const auto IsMatchingItem = [&](const FLoudnessMetricRef& LoudnessMetric) { return LoudnessMetric->Name == LoudnessMetricName; };
-		while (const FLoudnessMetricRef* MatchingItem = MeterValues->FindByPredicate(IsMatchingItem))
+		for (const FLoudnessMetricRef& LoudnessMetricRef : LoudnessMetrics.Get())
 		{
-			MeterValues->RemoveSingle(*MatchingItem);
+			const bool bShouldShowValue = LoudnessMetricRef->bShowValue.Get();
+			const bool bIsShowingValue = NumericValues->Contains(LoudnessMetricRef);
+			if (bIsShowingValue != bShouldShowValue)
+			{
+				if (bShouldShowValue)
+				{
+					NumericValues->Add(LoudnessMetricRef);
+				}
+				else
+				{
+					NumericValues->RemoveSingle(LoudnessMetricRef);
+				}
+			}
+
+			const bool bShouldShowMeter = LoudnessMetricRef->bShowMeter.Get();
+			const bool bIsShowingMeter = MeterValues->Contains(LoudnessMetricRef);
+			if (bIsShowingMeter != bShouldShowMeter)
+			{
+				if (bShouldShowMeter)
+				{
+					MeterValues->Add(LoudnessMetricRef);
+				}
+				else
+				{
+					MeterValues->RemoveSingle(LoudnessMetricRef);
+				}
+			}
 		}
 	}
 } // namespace AudioWidgets
diff --git a/Engine/Plugins/Runtime/AudioWidgets/Source/AudioWidgets/Private/LoudnessMeterWidgetView.h b/Engine/Plugins/Runtime/AudioWidgets/Source/AudioWidgets/Private/LoudnessMeterWidgetView.h
--- a/Engine/Plugins/Runtime/AudioWidgets/Source/AudioWidgets/Private/LoudnessMeterWidgetView.h
+++ b/Engine/Plugins/Runtime/AudioWidgets/Source/AudioWidgets/Private/LoudnessMeterWidgetView.h
@@ -20,6 +20,10 @@ namespace AudioWidgets
 			FName Name;
 			FText DisplayName;
 			TAttribute<TOptional<float>> Value;
+			TAttribute<bool> bShowValue = false;
+			TAttribute<bool> bShowMeter = false;
+			FSimpleDelegate OnShowValueToggleRequested;
+			FSimpleDelegate OnShowMeterToggleRequested;
 		};
 
 		struct FTimerPanelParams
@@ -37,16 +41,14 @@ namespace AudioWidgets
 
 		TSharedRef<SWidget> MakeWidget() const;
 
-		void AddNumericValue(const FLoudnessMetric& LoudnessMetric);
-		void RemoveNumericValue(FName LoudnessMetricName);
-
-		void AddMeterValue(const FLoudnessMetric& LoudnessMetric);
-		void RemoveMeterValue(FName LoudnessMetricName);
+		void AddLoudnessMetric(const FLoudnessMetric& LoudnessMetric);
+		void RefreshVisibleLoudnessMetrics();
 
 	private:
 		using FLoudnessMetricRef = TSharedRef<const FLoudnessMetric>;
 		using FLoudnessMetricRefArray = UE::Slate::Containers::TObservableArray<FLoudnessMetricRef>;
 
+		TSharedRef<FLoudnessMetricRefArray> LoudnessMetrics;
 		TSharedRef<FLoudnessMetricRefArray> NumericValues;
 		TSharedRef<FLoudnessMetricRefArray> MeterValues;
 		FTimerPanelParams TimerPanelParams;
diff --git a/Engine/Plugins/Runtime/AudioWidgets/Source/AudioWidgets/Public/LoudnessMeterRackUnitSettings.h b/Engine/Plugins/Runtime/AudioWidgets/Source/AudioWidgets/Public/LoudnessMeterRackUnitSettings.h
--- a/Engine/Plugins/Runtime/AudioWidgets/Source/AudioWidgets/Public/LoudnessMeterRackUnitSettings.h
+++ b/Engine/Plugins/Runtime/AudioWidgets/Source/AudioWidgets/Public/LoudnessMeterRackUnitSettings.h
@@ -3,6 +3,18 @@
 
 #include "LoudnessMeterRackUnitSettings.generated.h"
 
+USTRUCT()
+struct FLoudnessMetricDisplayOptions
+{
+	GENERATED_BODY()
+
+	UPROPERTY(EditAnywhere, config, Category = LoudnessMeter)
+	bool bShowValue = false;
+
+	UPROPERTY(EditAnywhere, config, Category = LoudnessMeter)
+	bool bShowMeter = false;
+};
+
 USTRUCT()
 struct FLoudnessMeterRackUnitSettings
 {
@@ -10,4 +22,13 @@ struct FLoudnessMeterRackUnitSettings
 
 	UPROPERTY(EditAnywhere, config, Category = LoudnessMeter)
 	bool bDisplayAnalysisTimer = true;
+
+	UPROPERTY(EditAnywhere, config, Category = LoudnessMeter)
+	FLoudnessMetricDisplayOptions IntegratedLoudness = { .bShowValue = true, .bShowMeter = false };
+
+	UPROPERTY(EditAnywhere, config, Category = LoudnessMeter)
+	FLoudnessMetricDisplayOptions ShortTermLoudness = { .bShowValue = false, .bShowMeter = true };
+
+	UPROPERTY(EditAnywhere, config, Category = LoudnessMeter)
+	FLoudnessMetricDisplayOptions MomentaryLoudness = { .bShowValue = false, .bShowMeter = true };
 };
