[Audio Widgets] Loudness Meter - Change name of Integrated Loudness to Long Term Loudness. Disable selection on loudness metric numeric value tile view.

diff --git a/Engine/Plugins/Runtime/AudioWidgets/Source/AudioWidgets/Private/AudioLoudnessMeter.cpp b/Engine/Plugins/Runtime/AudioWidgets/Source/AudioWidgets/Private/AudioLoudnessMeter.cpp
--- a/Engine/Plugins/Runtime/AudioWidgets/Source/AudioWidgets/Private/AudioLoudnessMeter.cpp
+++ b/Engine/Plugins/Runtime/AudioWidgets/Source/AudioWidgets/Private/AudioLoudnessMeter.cpp
@@ -166,8 +166,8 @@ namespace AudioWidgets
 		FLoudnessMeterWidgetView::FLoudnessMetric LoudnessMetrics[] =
 		{
 			{
-				.Name = GET_MEMBER_NAME_CHECKED(FLoudnessMeterRackUnitSettings, IntegratedLoudness),
-				.DisplayName = LOCTEXT("IntegratedLoudnessDisplayName", "Integrated"),
+				.Name = GET_MEMBER_NAME_CHECKED(FLoudnessMeterRackUnitSettings, LongTermLoudness),
+				.DisplayName = LOCTEXT("LongTermLoudnessDisplayName", "Long Term"),
 				.Value = TAttribute<TOptional<float>>::CreateSPLambda(this, [this]() { return (LatestLoudnessResults.IsSet()) ? TOptional<float>(LatestLoudnessResults->GatedLoudness) : NullOpt; }),
 				.bShowValue = true,
 				.bShowMeter = false,
@@ -261,9 +261,10 @@ namespace AudioWidgets
 
 	FTimespan FAudioLoudnessMeter::GetAnalysisTime() const
 	{
-		if (LatestLoudnessResults.IsSet())
+		if (LatestLoudnessResults.IsSet() && LKFSSettings.IsValid())
 		{
-			return FTimespan::FromSeconds(LatestLoudnessResults->Timestamp);
+			const float AnalysisTime = FMath::Min(LatestLoudnessResults->Timestamp, LKFSSettings->IntegratedLoudnessDuration);
+			return FTimespan::FromSeconds(AnalysisTime);
 		}
 	
 		return FTimespan::Zero();
diff --git a/Engine/Plugins/Runtime/AudioWidgets/Source/AudioWidgets/Private/LoudnessMeterWidgetView.cpp b/Engine/Plugins/Runtime/AudioWidgets/Source/AudioWidgets/Private/LoudnessMeterWidgetView.cpp
--- a/Engine/Plugins/Runtime/AudioWidgets/Source/AudioWidgets/Private/LoudnessMeterWidgetView.cpp
+++ b/Engine/Plugins/Runtime/AudioWidgets/Source/AudioWidgets/Private/LoudnessMeterWidgetView.cpp
@@ -334,7 +334,7 @@ namespace AudioWidgets
 												SNew(SImage)
 													.Image(FAudioWidgetsStyle::Get().GetBrush("AudioWidgetsStyle.Reset"))
 													.ColorAndOpacity(FSlateColor::UseForeground())
-													.ToolTipText(LOCTEXT("ResetButtonToolTip", "Reset Integrated loudness analysis"))
+													.ToolTipText(LOCTEXT("ResetButtonToolTip", "Reset Long Term loudness analysis"))
 											]
 									]
 							]
@@ -360,6 +360,7 @@ namespace AudioWidgets
 							.ItemHeight(60.0f)
 							.ItemWidth(60.0f)
 							.ItemAlignment(EListItemAlignment::CenterAligned)
+							.SelectionMode(ESelectionMode::None)
 					]
 					+ SVerticalBox::Slot()
 					.AutoHeight()
diff --git a/Engine/Plugins/Runtime/AudioWidgets/Source/AudioWidgets/Public/LoudnessMeterRackUnitSettings.h b/Engine/Plugins/Runtime/AudioWidgets/Source/AudioWidgets/Public/LoudnessMeterRackUnitSettings.h
--- a/Engine/Plugins/Runtime/AudioWidgets/Source/AudioWidgets/Public/LoudnessMeterRackUnitSettings.h
+++ b/Engine/Plugins/Runtime/AudioWidgets/Source/AudioWidgets/Public/LoudnessMeterRackUnitSettings.h
@@ -24,7 +24,7 @@ struct FLoudnessMeterRackUnitSettings
 	bool bDisplayAnalysisTimer = true;
 
 	UPROPERTY(EditAnywhere, config, Category = LoudnessMeter)
-	FLoudnessMetricDisplayOptions IntegratedLoudness = { .bShowValue = true, .bShowMeter = false };
+	FLoudnessMetricDisplayOptions LongTermLoudness = { .bShowValue = true, .bShowMeter = false };
 
 	UPROPERTY(EditAnywhere, config, Category = LoudnessMeter)
 	FLoudnessMetricDisplayOptions ShortTermLoudness = { .bShowValue = false, .bShowMeter = true };
