Fixing audio slider binding

diff --git a/Engine/Plugins/Runtime/AudioWidgets/Source/AudioWidgets/Private/AudioSlider.cpp b/Engine/Plugins/Runtime/AudioWidgets/Source/AudioWidgets/Private/AudioSlider.cpp
--- a/Engine/Plugins/Runtime/AudioWidgets/Source/AudioWidgets/Private/AudioSlider.cpp
+++ b/Engine/Plugins/Runtime/AudioWidgets/Source/AudioWidgets/Private/AudioSlider.cpp
@@ -66,18 +66,31 @@ UAudioSliderBase::UAudioSliderBase(const FObjectInitializer& ObjectInitializer)
 void UAudioSliderBase::SynchronizeProperties()
 {
 	Super::SynchronizeProperties();
-	
-	MyAudioSlider->SetLabelBackgroundColor(TextLabelBackgroundColor);
-	MyAudioSlider->SetUnitsText(UnitsText);
-	MyAudioSlider->SetUnitsTextReadOnly(IsUnitsTextReadOnly);
-	MyAudioSlider->SetValueTextReadOnly(IsValueTextReadOnly);
-	MyAudioSlider->SetShowLabelOnlyOnHover(ShowLabelOnlyOnHover);
-	MyAudioSlider->SetShowUnitsText(ShowUnitsText);
-	MyAudioSlider->SetOrientation(Orientation);
-	MyAudioSlider->SetSliderBackgroundColor(SliderBackgroundColor);
-	MyAudioSlider->SetSliderBarColor(SliderBarColor);
-	MyAudioSlider->SetSliderThumbColor(SliderThumbColor);
-	MyAudioSlider->SetWidgetBackgroundColor(WidgetBackgroundColor);
+
+	if (MyAudioSlider.IsValid())
+	{
+		MyAudioSlider->SetLabelBackgroundColor(TextLabelBackgroundColor);
+		MyAudioSlider->SetUnitsText(UnitsText);
+		MyAudioSlider->SetUnitsTextReadOnly(IsUnitsTextReadOnly);
+		MyAudioSlider->SetValueTextReadOnly(IsValueTextReadOnly);
+		MyAudioSlider->SetShowLabelOnlyOnHover(ShowLabelOnlyOnHover);
+		MyAudioSlider->SetShowUnitsText(ShowUnitsText);
+		MyAudioSlider->SetOrientation(Orientation);
+		MyAudioSlider->SetSliderBackgroundColor(SliderBackgroundColor);
+		MyAudioSlider->SetSliderBarColor(SliderBarColor);
+		MyAudioSlider->SetSliderThumbColor(SliderThumbColor);
+		MyAudioSlider->SetWidgetBackgroundColor(WidgetBackgroundColor);
+
+		TAttribute<float> ValueBinding = PROPERTY_BINDING(float, Value);
+		if (ValueBinding.IsBound())
+		{
+			MyAudioSlider->SetSliderValue(ValueBinding);
+		}
+		else
+		{
+			MyAudioSlider->SetSliderValue(Value);
+		}
+	}
 }
 
 void UAudioSliderBase::ReleaseSlateResources(bool bReleaseChildren)
@@ -89,6 +102,7 @@ void UAudioSliderBase::ReleaseSlateResources(bool bReleaseChildren)
 
 void UAudioSliderBase::HandleOnValueChanged(float InValue)
 {
+	Value = InValue;
 	OnValueChanged.Broadcast(InValue);
 }
 
diff --git a/Engine/Plugins/Runtime/AudioWidgets/Source/AudioWidgets/Private/SAudioSlider.cpp b/Engine/Plugins/Runtime/AudioWidgets/Source/AudioWidgets/Private/SAudioSlider.cpp
--- a/Engine/Plugins/Runtime/AudioWidgets/Source/AudioWidgets/Private/SAudioSlider.cpp
+++ b/Engine/Plugins/Runtime/AudioWidgets/Source/AudioWidgets/Private/SAudioSlider.cpp
@@ -50,13 +50,13 @@ void SAudioSliderBase::Construct(const SAudioSliderBase::FArguments& InArgs)
 
 	// Underlying slider widget
 	SAssignNew(Slider, SSlider)
-		.Value(SliderValueAttribute.Get())
+		.Value(SliderValueAttribute)
 		.Style(&Style->SliderStyle)
 		.Orientation(Orientation.Get())
 		.IndentHandle(false)
 		.OnValueChanged_Lambda([this](float Value)
 		{
-			SliderValueAttribute.Set(Value);
+			SetSliderValue(Value);
 			OnValueChanged.ExecuteIfBound(Value);
 			const float OutputValue = GetOutputValueForText(Value);
 			Label->SetValueText(OutputValue);
@@ -75,7 +75,7 @@ void SAudioSliderBase::Construct(const SAudioSliderBase::FArguments& InArgs)
 			const float NewSliderValue = GetSliderValueForText(OutputValue);
 			if (!FMath::IsNearlyEqual(NewSliderValue, SliderValueAttribute.Get()))
 			{
-				SliderValueAttribute.Set(NewSliderValue);
+				SetSliderValue(NewSliderValue);
 				Slider->SetValue(NewSliderValue);
 				OnValueChanged.ExecuteIfBound(NewSliderValue);
 				OnValueCommitted.ExecuteIfBound(NewSliderValue);
@@ -96,6 +96,18 @@ void SAudioSliderBase::SetSliderValue(float InSliderValue)
 	Slider->SetValue(InSliderValue);
 }
 
+void SAudioSliderBase::SetSliderValue(TAttribute<float> InSliderValue)
+{
+	SliderValueAttribute = MoveTemp(InSliderValue);
+	if (Slider.IsValid())
+	{
+		Slider->SetValue(SliderValueAttribute);
+	}
+
+	const float Current = SliderValueAttribute.Get();
+	Label->SetValueText(GetOutputValueForText(Current));
+}
+
 void SAudioSliderBase::SetOrientation(EOrientation InOrientation)
 {
 	SetAttribute(Orientation, TAttribute<EOrientation>(InOrientation), EInvalidateWidgetReason::Layout);
@@ -190,12 +202,35 @@ void SAudioSliderBase::SetOutputRange(const FVector2D Range)
 	// if Range.Y < Range.X, set Range.X to Range.Y
 	OutputRange.X = FMath::Min(Range.X, Range.Y);
 
-	const float OutputValue = GetOutputValue(SliderValueAttribute.Get());
-	const float ClampedOutputValue = FMath::Clamp(OutputValue, OutputRange.X, OutputRange.Y);
-	const float ClampedSliderValue = GetSliderValue(ClampedOutputValue);
-	SetSliderValue(ClampedSliderValue);
+	// Cache these values to handle case where we have a bound slider value attribute
+	const float OldOutputValueNormalized = SliderValueAttribute.Get();
+	const float OldOutputValueMapped = GetOutputValue(OldOutputValueNormalized);
+
+	const float NewClampedOutputValue = FMath::Clamp(OldOutputValueMapped, OutputRange.X, OutputRange.Y);
+	const float NewNormalizedValue = GetSliderValue(NewClampedOutputValue);
 
+	Label->SetValueText(GetOutputValueForText(NewNormalizedValue));
 	Label->UpdateValueTextWidth(OutputRange);
+
+	// Avoid mutating the SliderValueAttribute if we're bound to avoid severing the binding
+	if (SliderValueAttribute.IsBound())
+	{
+		if (!FMath::IsNearlyEqual(OldOutputValueNormalized, NewNormalizedValue))
+		{
+			OnValueChanged.ExecuteIfBound(NewNormalizedValue);
+			OnValueCommitted.ExecuteIfBound(NewNormalizedValue);
+		}
+	}
+	else
+	{
+		SliderValueAttribute.Set(NewNormalizedValue);
+		if (Slider.IsValid())
+		{
+			Slider->SetValue(SliderValueAttribute);
+		}
+		OnValueChanged.ExecuteIfBound(NewNormalizedValue);
+		OnValueCommitted.ExecuteIfBound(NewNormalizedValue);
+	}
 }
 
 void SAudioSliderBase::SetLabelBackgroundColor(FSlateColor InColor)
diff --git a/Engine/Plugins/Runtime/AudioWidgets/Source/AudioWidgets/Public/SAudioSlider.h b/Engine/Plugins/Runtime/AudioWidgets/Source/AudioWidgets/Public/SAudioSlider.h
--- a/Engine/Plugins/Runtime/AudioWidgets/Source/AudioWidgets/Public/SAudioSlider.h
+++ b/Engine/Plugins/Runtime/AudioWidgets/Source/AudioWidgets/Public/SAudioSlider.h
@@ -99,6 +99,7 @@ public:
 	 * Set the slider's linear (0-1 normalized) value. 
 	 */
 	UE_API void SetSliderValue(float InSliderValue);
+	UE_API void SetSliderValue(TAttribute<float> InSliderValue);
 	UE_API FVector2D ComputeDesiredSize(float) const;
 	UE_API void SetDesiredSizeOverride(const FVector2D DesiredSize);
 
