Fix undo not working in MetaSounds after undoing a change using an audio material slider

diff --git a/Engine/Plugins/Runtime/AudioWidgets/Source/AudioWidgets/Private/AudioMaterialSlate/SAudioMaterialLabeledSlider.cpp b/Engine/Plugins/Runtime/AudioWidgets/Source/AudioWidgets/Private/AudioMaterialSlate/SAudioMaterialLabeledSlider.cpp
--- a/Engine/Plugins/Runtime/AudioWidgets/Source/AudioWidgets/Private/AudioMaterialSlate/SAudioMaterialLabeledSlider.cpp
+++ b/Engine/Plugins/Runtime/AudioWidgets/Source/AudioWidgets/Private/AudioMaterialSlate/SAudioMaterialLabeledSlider.cpp
@@ -48,6 +48,10 @@ void SAudioMaterialLabeledSlider::Construct(const SAudioMaterialLabeledSlider::F
 			OnValueChanged.ExecuteIfBound(Value);
 			const float OutputValue = GetOutputValueForText(Value);
 			Label->SetValueText(OutputValue);
+		})
+		.OnValueCommitted_Lambda([this](float Value)
+		{
+			OnValueCommitted.ExecuteIfBound(SliderValueAttribute.Get());
 		});
 
 	//Set the Unit processor from Units Type
