Fix non material MetaSound slider value snapping to 0

diff --git a/Engine/Plugins/Runtime/AudioWidgets/Source/AudioWidgets/Private/SAudioSlider.cpp b/Engine/Plugins/Runtime/AudioWidgets/Source/AudioWidgets/Private/SAudioSlider.cpp
--- a/Engine/Plugins/Runtime/AudioWidgets/Source/AudioWidgets/Private/SAudioSlider.cpp
+++ b/Engine/Plugins/Runtime/AudioWidgets/Source/AudioWidgets/Private/SAudioSlider.cpp
@@ -212,24 +212,24 @@ void SAudioSliderBase::SetOutputRange(const FVector2D Range)
 	Label->SetValueText(GetOutputValueForText(NewNormalizedValue));
 	Label->UpdateValueTextWidth(OutputRange);
 
-	// Avoid mutating the SliderValueAttribute if we're bound to avoid severing the binding
-	if (SliderValueAttribute.IsBound())
+	if (!FMath::IsNearlyEqual(OldOutputValueNormalized, NewNormalizedValue))
 	{
-		if (!FMath::IsNearlyEqual(OldOutputValueNormalized, NewNormalizedValue))
+		// Avoid mutating the SliderValueAttribute if we're bound to avoid severing the binding
+		if (SliderValueAttribute.IsBound())
 		{
 			OnValueChanged.ExecuteIfBound(NewNormalizedValue);
 			OnValueCommitted.ExecuteIfBound(NewNormalizedValue);
 		}
-	}
-	else
-	{
-		SliderValueAttribute.Set(NewNormalizedValue);
-		if (Slider.IsValid())
+		else
 		{
-			Slider->SetValue(SliderValueAttribute);
+			SliderValueAttribute.Set(NewNormalizedValue);
+			if (Slider.IsValid())
+			{
+				Slider->SetValue(SliderValueAttribute);
+			}
+			OnValueChanged.ExecuteIfBound(NewNormalizedValue);
+			OnValueCommitted.ExecuteIfBound(NewNormalizedValue);
 		}
-		OnValueChanged.ExecuteIfBound(NewNormalizedValue);
-		OnValueCommitted.ExecuteIfBound(NewNormalizedValue);
 	}
 }
 
