[Audio Widgets] Loudness Meter - Added menu options for Value Color and Reset Color for the loudness metrics.

diff --git a/Engine/Plugins/Runtime/AudioWidgets/Source/AudioWidgets/AudioWidgets.Build.cs b/Engine/Plugins/Runtime/AudioWidgets/Source/AudioWidgets/AudioWidgets.Build.cs
--- a/Engine/Plugins/Runtime/AudioWidgets/Source/AudioWidgets/AudioWidgets.Build.cs
+++ b/Engine/Plugins/Runtime/AudioWidgets/Source/AudioWidgets/AudioWidgets.Build.cs
@@ -24,6 +24,7 @@ public class AudioWidgets : ModuleRules
 		(
 			new string[]
 			{
+				"AppFramework",
 				"AudioMixer",
 				"CoreUObject",
 				"Engine",
diff --git a/Engine/Plugins/Runtime/AudioWidgets/Source/AudioWidgets/Private/AudioLoudnessMeter.cpp b/Engine/Plugins/Runtime/AudioWidgets/Source/AudioWidgets/Private/AudioLoudnessMeter.cpp
--- a/Engine/Plugins/Runtime/AudioWidgets/Source/AudioWidgets/Private/AudioLoudnessMeter.cpp
+++ b/Engine/Plugins/Runtime/AudioWidgets/Source/AudioWidgets/Private/AudioLoudnessMeter.cpp
@@ -207,31 +207,57 @@ namespace AudioWidgets
 				{
 					if (const FProperty* DisplayOptionsProperty = FLoudnessMeterRackUnitSettings::StaticStruct()->FindPropertyByName(LoudnessMetric.Name))
 					{
+						const auto GetLoudnessMetricDisplayOptions = [DisplayOptionsProperty](FLoudnessMeterRackUnitSettings& RackUnitSettings)
+							{
+								return DisplayOptionsProperty->ContainerPtrToValuePtr<FLoudnessMetricDisplayOptions>(&RackUnitSettings);
+							};
+
 						LoudnessMetric.bShowValue.BindLambda([=]()
 							{
-								const FLoudnessMetricDisplayOptions* DisplayOptions = DisplayOptionsProperty->ContainerPtrToValuePtr<FLoudnessMetricDisplayOptions>(&SettingsHelper.GetRackUnitSettings());
+								const FLoudnessMetricDisplayOptions* DisplayOptions = GetLoudnessMetricDisplayOptions(SettingsHelper.GetRackUnitSettings());
 								return DisplayOptions->bShowValue;
 							});
 
 						LoudnessMetric.bShowMeter.BindLambda([=]()
 							{
-								const FLoudnessMetricDisplayOptions* DisplayOptions = DisplayOptionsProperty->ContainerPtrToValuePtr<FLoudnessMetricDisplayOptions>(&SettingsHelper.GetRackUnitSettings());
+								const FLoudnessMetricDisplayOptions* DisplayOptions = GetLoudnessMetricDisplayOptions(SettingsHelper.GetRackUnitSettings());
 								return DisplayOptions->bShowMeter;
 							});
 
 						LoudnessMetric.OnShowValueToggleRequested.BindLambda([=]()
 							{
-								FLoudnessMetricDisplayOptions* DisplayOptions = DisplayOptionsProperty->ContainerPtrToValuePtr<FLoudnessMetricDisplayOptions>(&SettingsHelper.GetRackUnitSettings());
+								FLoudnessMetricDisplayOptions* DisplayOptions = GetLoudnessMetricDisplayOptions(SettingsHelper.GetRackUnitSettings());
 								DisplayOptions->bShowValue = !DisplayOptions->bShowValue;
 								SettingsHelper.SaveConfig();
 							});
 
 						LoudnessMetric.OnShowMeterToggleRequested.BindLambda([=]()
 							{
-								FLoudnessMetricDisplayOptions* DisplayOptions = DisplayOptionsProperty->ContainerPtrToValuePtr<FLoudnessMetricDisplayOptions>(&SettingsHelper.GetRackUnitSettings());
+								FLoudnessMetricDisplayOptions* DisplayOptions = GetLoudnessMetricDisplayOptions(SettingsHelper.GetRackUnitSettings());
 								DisplayOptions->bShowMeter = !DisplayOptions->bShowMeter;
 								SettingsHelper.SaveConfig();
 							});
+
+						LoudnessMetric.Color.BindLambda([=]()
+							{
+								const FLoudnessMetricDisplayOptions* DisplayOptions = GetLoudnessMetricDisplayOptions(SettingsHelper.GetRackUnitSettings());
+								return DisplayOptions->Color;
+							});
+
+						LoudnessMetric.OnColorChanged.BindLambda([=](FLinearColor NewColor)
+							{
+								FLoudnessMetricDisplayOptions* DisplayOptions = GetLoudnessMetricDisplayOptions(SettingsHelper.GetRackUnitSettings());
+								DisplayOptions->Color = NewColor;
+							});
+
+						LoudnessMetric.OnColorPickerWindowClosed.BindLambda([=](const TSharedRef<SWindow>& Window) { SettingsHelper.SaveConfig(); });
+
+						LoudnessMetric.OnResetColorRequested.BindLambda([=]()
+							{
+								FLoudnessMetricDisplayOptions* DisplayOptions = GetLoudnessMetricDisplayOptions(SettingsHelper.GetRackUnitSettings());
+								DisplayOptions->Color = FLoudnessMeterWidgetView::DefaultMeterColor;
+								SettingsHelper.SaveConfig();
+							});
 					}
 				}
 
diff --git a/Engine/Plugins/Runtime/AudioWidgets/Source/AudioWidgets/Private/AudioWidgetsStyle.cpp b/Engine/Plugins/Runtime/AudioWidgets/Source/AudioWidgets/Private/AudioWidgetsStyle.cpp
--- a/Engine/Plugins/Runtime/AudioWidgets/Source/AudioWidgets/Private/AudioWidgetsStyle.cpp
+++ b/Engine/Plugins/Runtime/AudioWidgets/Source/AudioWidgets/Private/AudioWidgetsStyle.cpp
@@ -195,6 +195,8 @@ void FAudioWidgetsStyle::SetupStyles()
 	const FVector2f Icon16x16 = FVector2f(16.0f, 16.0f);
 	const FVector2f Icon20x20 = FVector2f(20.0f, 20.0f);
 
+	Set("AudioWidgetsStyle.ColorMixing", new IMAGE_BRUSH_SVG(TEXT("Slate/Icons/ColorMixing"), Icon16x16));
+
 	Set("AudioWidgetsStyle.FadeIn", new IMAGE_BRUSH_SVG(TEXT("Slate/Icons/FadeIn_20"), Icon20x20));
 	Set("AudioWidgetsStyle.FadeInLinear", new IMAGE_BRUSH_SVG(TEXT("Slate/Icons/FadeInLinear_16"), Icon16x16));
 	Set("AudioWidgetsStyle.FadeInLogarithmic", new IMAGE_BRUSH_SVG(TEXT("Slate/Icons/FadeInLogarithmic_16"), Icon16x16));
diff --git a/Engine/Plugins/Runtime/AudioWidgets/Source/AudioWidgets/Private/LoudnessMeterWidgetView.cpp b/Engine/Plugins/Runtime/AudioWidgets/Source/AudioWidgets/Private/LoudnessMeterWidgetView.cpp
--- a/Engine/Plugins/Runtime/AudioWidgets/Source/AudioWidgets/Private/LoudnessMeterWidgetView.cpp
+++ b/Engine/Plugins/Runtime/AudioWidgets/Source/AudioWidgets/Private/LoudnessMeterWidgetView.cpp
@@ -2,14 +2,17 @@
 
 #include "LoudnessMeterWidgetView.h"
 
+#include "AudioDefines.h"
 #include "AudioWidgetsStyle.h"
 #include "AudioMeterWidgetStyle.h"
 #include "Brushes/SlateColorBrush.h"
 #include "Brushes/SlateRoundedBoxBrush.h"
+#include "Framework/Application/SlateApplication.h"
 #include "Framework/MultiBox/MultiBoxBuilder.h"
 #include "SAudioMeterWidget.h"
 #include "Styling/StyleColors.h"
 #include "Widgets/Colors/SColorBlock.h"
+#include "Widgets/Colors/SColorPicker.h"
 #include "Widgets/Input/SButton.h"
 #include "Widgets/Layout/SSeparator.h"
 #include "Widgets/Layout/SSpacer.h"
@@ -62,7 +65,7 @@ namespace AudioWidgets
 							SNew(STextBlock)
 								.Text_Lambda([InItem]() { return FormatNumericValueText(InItem->Value.Get()); })
 								.Font(FStyleDefaults::GetFontInfo(14))
-								.ColorAndOpacity(FStyleColors::AccentGreen)
+								.ColorAndOpacity_Lambda([InItem]() { return InItem->Color.Get().CopyWithNewOpacity(1.0f); })
 								.SimpleTextMode(true)
 						]
 						+ SVerticalBox::Slot()
@@ -117,7 +120,7 @@ namespace AudioWidgets
 											SNew(STextBlock)
 												.Text_Lambda([InItem]() { return FormatNumericValueText(InItem->Value.Get()); })
 												.Font(FStyleDefaults::GetFontInfo(14))
-												.ColorAndOpacity(FStyleColors::AccentGreen)
+												.ColorAndOpacity_Lambda([InItem]() { return InItem->Color.Get().CopyWithNewOpacity(1.0f); })
 												.SimpleTextMode(true)
 										]
 										+ SVerticalBox::Slot()
@@ -127,7 +130,7 @@ namespace AudioWidgets
 												.Orientation(EOrientation::Orient_Vertical)
 												.BackgroundColor(FLinearColor::Transparent)
 												.MeterBackgroundColor(DefaultColorStyle.MeterBackgroundColor)
-												.MeterValueColor(DefaultColorStyle.MeterValueColor)
+												.MeterValueColor_Lambda([InItem]() { return InItem->Color.Get(); })
 												.MeterPeakColor(FLinearColor::Transparent)
 												.MeterClippingColor(FLinearColor::Transparent)
 												.MeterScaleColor(FLinearColor::White.CopyWithNewOpacity(0.25f))
@@ -135,7 +138,7 @@ namespace AudioWidgets
 												.Style(&AudioMeterWidgetStyle)
 												.MeterChannelInfo_Lambda([InItem]() -> TArray<FAudioMeterChannelInfo>
 													{
-														constexpr float DefaultInvalidValue = -160.0f;
+														constexpr float DefaultInvalidValue = MIN_VOLUME_DECIBELS;
 														const TOptional<float> LoudnessMetricValue = InItem->Value.Get();
 														return { { LoudnessMetricValue.Get(DefaultInvalidValue), DefaultInvalidValue, DefaultInvalidValue } };
 													})
@@ -158,6 +161,8 @@ namespace AudioWidgets
 		}
 	} // namespace LoudnessMeterWidgetView_Private
 
+	const FLinearColor FLoudnessMeterWidgetView::DefaultMeterColor = FStyleColors::AccentGreen.GetSpecifiedColor().CopyWithNewOpacity(0.5f);
+
 	FLoudnessMeterWidgetView::FLoudnessMeterWidgetView()
 		: LoudnessMetrics(MakeShared<FLoudnessMetricRefArray>())
 		, NumericValues(MakeShared<FLoudnessMetricRefArray>())
@@ -241,6 +246,39 @@ namespace AudioWidgets
 										NAME_None,
 										EUserInterfaceActionType::ToggleButton);
 
+									if (LoudnessMetric.OnColorChanged.IsBound())
+									{
+										SubMenu.AddMenuEntry(
+											LOCTEXT("LayoutOptions_LoudnessMetric_ValueColor", "Value Color"),
+											FText(),
+											FSlateIcon(FAudioWidgetsStyle::StyleName, "AudioWidgetsStyle.ColorMixing"),
+											FExecuteAction::CreateLambda([
+												ParentWidget = FSlateApplication::Get().GetMenuHostWidget().ToWeakPtr(),
+												Color = LoudnessMetric.Color,
+												OnColorChanged = LoudnessMetric.OnColorChanged,
+												OnColorPickerWindowClosed = LoudnessMetric.OnColorPickerWindowClosed]()
+												{
+													FColorPickerArgs PickerArgs(Color.Get(), OnColorChanged);
+													PickerArgs.ParentWidget = ParentWidget.Pin();
+													PickerArgs.bUseAlpha = true;
+													PickerArgs.OnColorPickerWindowClosed = OnColorPickerWindowClosed;
+													OpenColorPicker(PickerArgs);
+												}),
+											NAME_None,
+											EUserInterfaceActionType::Button);
+									}
+
+									if (LoudnessMetric.OnResetColorRequested.IsBound())
+									{
+										SubMenu.AddMenuEntry(
+											LOCTEXT("LayoutOptions_LoudnessMetric_ResetColor", "Reset Color"),
+											FText(),
+											FSlateIcon(FAudioWidgetsStyle::StyleName, "AudioWidgetsStyle.Reset"),
+											LoudnessMetric.OnResetColorRequested,
+											NAME_None,
+											EUserInterfaceActionType::Button);
+									}
+
 									SubMenu.EndSection();
 								}));
 					}
diff --git a/Engine/Plugins/Runtime/AudioWidgets/Source/AudioWidgets/Private/LoudnessMeterWidgetView.h b/Engine/Plugins/Runtime/AudioWidgets/Source/AudioWidgets/Private/LoudnessMeterWidgetView.h
--- a/Engine/Plugins/Runtime/AudioWidgets/Source/AudioWidgets/Private/LoudnessMeterWidgetView.h
+++ b/Engine/Plugins/Runtime/AudioWidgets/Source/AudioWidgets/Private/LoudnessMeterWidgetView.h
@@ -5,6 +5,7 @@
 #include "Containers/ObservableArray.h"
 #include "Framework/SlateDelegates.h"
 #include "Widgets/SWidget.h"
+#include "Widgets/SWindow.h"
 
 namespace AudioWidgets
 {
@@ -15,6 +16,8 @@ namespace AudioWidgets
 	class FLoudnessMeterWidgetView
 	{
 	public:
+		static const FLinearColor DefaultMeterColor;
+
 		struct FLoudnessMetric
 		{
 			FName Name;
@@ -24,6 +27,10 @@ namespace AudioWidgets
 			TAttribute<bool> bShowMeter = false;
 			FSimpleDelegate OnShowValueToggleRequested;
 			FSimpleDelegate OnShowMeterToggleRequested;
+			TAttribute<FLinearColor> Color = DefaultMeterColor;
+			FOnLinearColorValueChanged OnColorChanged;
+			FOnWindowClosed OnColorPickerWindowClosed;
+			FSimpleDelegate OnResetColorRequested;
 		};
 
 		struct FTimerPanelParams
diff --git a/Engine/Plugins/Runtime/AudioWidgets/Source/AudioWidgets/Public/LoudnessMeterRackUnitSettings.h b/Engine/Plugins/Runtime/AudioWidgets/Source/AudioWidgets/Public/LoudnessMeterRackUnitSettings.h
--- a/Engine/Plugins/Runtime/AudioWidgets/Source/AudioWidgets/Public/LoudnessMeterRackUnitSettings.h
+++ b/Engine/Plugins/Runtime/AudioWidgets/Source/AudioWidgets/Public/LoudnessMeterRackUnitSettings.h
@@ -1,6 +1,7 @@
 // Copyright Epic Games, Inc. All Rights Reserved.
 #pragma once
 
+#include "Styling/StyleColors.h"
 #include "LoudnessMeterRackUnitSettings.generated.h"
 
 USTRUCT()
@@ -13,6 +14,9 @@ struct FLoudnessMetricDisplayOptions
 
 	UPROPERTY(EditAnywhere, config, Category = LoudnessMeter)
 	bool bShowMeter = false;
+
+	UPROPERTY(EditAnywhere, config, Category = LoudnessMeter)
+	FLinearColor Color = FStyleColors::AccentGreen.GetSpecifiedColor().CopyWithNewOpacity(0.5f);
 };
 
 USTRUCT()
