Fix crash regarding FLoudnessMeterWidgetView::DefaultMeterColor.

diff --git a/Engine/Plugins/Runtime/AudioWidgets/Source/AudioWidgets/Private/AudioLoudnessMeter.cpp b/Engine/Plugins/Runtime/AudioWidgets/Source/AudioWidgets/Private/AudioLoudnessMeter.cpp
--- a/Engine/Plugins/Runtime/AudioWidgets/Source/AudioWidgets/Private/AudioLoudnessMeter.cpp
+++ b/Engine/Plugins/Runtime/AudioWidgets/Source/AudioWidgets/Private/AudioLoudnessMeter.cpp
@@ -255,7 +255,7 @@ namespace AudioWidgets
 						LoudnessMetric.OnResetColorRequested.BindLambda([=]()
 							{
 								FLoudnessMetricDisplayOptions* DisplayOptions = GetLoudnessMetricDisplayOptions(SettingsHelper.GetRackUnitSettings());
-								DisplayOptions->Color = FLoudnessMeterWidgetView::DefaultMeterColor;
+								DisplayOptions->Color = FLoudnessMeterWidgetView::GetDefaultMeterColor();
 								SettingsHelper.SaveConfig();
 							});
 					}
diff --git a/Engine/Plugins/Runtime/AudioWidgets/Source/AudioWidgets/Private/LoudnessMeterWidgetView.cpp b/Engine/Plugins/Runtime/AudioWidgets/Source/AudioWidgets/Private/LoudnessMeterWidgetView.cpp
--- a/Engine/Plugins/Runtime/AudioWidgets/Source/AudioWidgets/Private/LoudnessMeterWidgetView.cpp
+++ b/Engine/Plugins/Runtime/AudioWidgets/Source/AudioWidgets/Private/LoudnessMeterWidgetView.cpp
@@ -161,7 +161,10 @@ namespace AudioWidgets
 		}
 	} // namespace LoudnessMeterWidgetView_Private
 
-	const FLinearColor FLoudnessMeterWidgetView::DefaultMeterColor = FStyleColors::AccentGreen.GetSpecifiedColor().CopyWithNewOpacity(0.5f);
+	FLinearColor FLoudnessMeterWidgetView::GetDefaultMeterColor()
+	{
+		return FStyleColors::AccentGreen.GetSpecifiedColor().CopyWithNewOpacity(0.5f);
+	}
 
 	FLoudnessMeterWidgetView::FLoudnessMeterWidgetView()
 		: LoudnessMetrics(MakeShared<FLoudnessMetricRefArray>())
diff --git a/Engine/Plugins/Runtime/AudioWidgets/Source/AudioWidgets/Private/LoudnessMeterWidgetView.h b/Engine/Plugins/Runtime/AudioWidgets/Source/AudioWidgets/Private/LoudnessMeterWidgetView.h
--- a/Engine/Plugins/Runtime/AudioWidgets/Source/AudioWidgets/Private/LoudnessMeterWidgetView.h
+++ b/Engine/Plugins/Runtime/AudioWidgets/Source/AudioWidgets/Private/LoudnessMeterWidgetView.h
@@ -16,7 +16,7 @@ namespace AudioWidgets
 	class FLoudnessMeterWidgetView
 	{
 	public:
-		static const FLinearColor DefaultMeterColor;
+		static FLinearColor GetDefaultMeterColor();
 
 		struct FLoudnessMetric
 		{
@@ -27,7 +27,7 @@ namespace AudioWidgets
 			TAttribute<bool> bShowMeter = false;
 			FSimpleDelegate OnShowValueToggleRequested;
 			FSimpleDelegate OnShowMeterToggleRequested;
-			TAttribute<FLinearColor> Color = DefaultMeterColor;
+			TAttribute<FLinearColor> Color = GetDefaultMeterColor();
 			FOnLinearColorValueChanged OnColorChanged;
 			FOnWindowClosed OnColorPickerWindowClosed;
 			FSimpleDelegate OnResetColorRequested;
