[Audio Widgets] Loudness Meter - added menu options for scale range and offset.

diff --git a/Engine/Plugins/Runtime/AudioWidgets/Source/AudioWidgets/Private/AudioLoudnessMeter.cpp b/Engine/Plugins/Runtime/AudioWidgets/Source/AudioWidgets/Private/AudioLoudnessMeter.cpp
--- a/Engine/Plugins/Runtime/AudioWidgets/Source/AudioWidgets/Private/AudioLoudnessMeter.cpp
+++ b/Engine/Plugins/Runtime/AudioWidgets/Source/AudioWidgets/Private/AudioLoudnessMeter.cpp
@@ -163,6 +163,8 @@ namespace AudioWidgets
 			.bIsVisible = true
 		};
 
+		FLoudnessMeterWidgetView::FLoudnessScaleParams LoudnessScaleParams;
+
 		FLoudnessMeterWidgetView::FLoudnessMetric LoudnessMetrics[] =
 		{
 			{
@@ -203,6 +205,22 @@ namespace AudioWidgets
 						SettingsHelper.SaveConfig();
 					});
 
+				LoudnessScaleParams.Range.BindLambda([=]() { return SettingsHelper.GetRackUnitSettings().LoudnessScaleRange; });
+				LoudnessScaleParams.OnRangeValueChanged.BindLambda([=](int Value) { SettingsHelper.GetRackUnitSettings().LoudnessScaleRange = Value; });
+				LoudnessScaleParams.OnRangeValueCommitted.BindLambda([=](int Value, ETextCommit::Type CommitType)
+					{
+						SettingsHelper.GetRackUnitSettings().LoudnessScaleRange = Value;
+						SettingsHelper.SaveConfig();
+					});
+
+				LoudnessScaleParams.Offset.BindLambda([=]() { return SettingsHelper.GetRackUnitSettings().LoudnessScaleOffset; });
+				LoudnessScaleParams.OnOffsetValueChanged.BindLambda([=](int Value) { SettingsHelper.GetRackUnitSettings().LoudnessScaleOffset = Value; });
+				LoudnessScaleParams.OnOffsetValueCommitted.BindLambda([=](int Value, ETextCommit::Type CommitType)
+					{
+						SettingsHelper.GetRackUnitSettings().LoudnessScaleOffset = Value;
+						SettingsHelper.SaveConfig();
+					});
+
 				for (FLoudnessMeterWidgetView::FLoudnessMetric& LoudnessMetric : LoudnessMetrics)
 				{
 					if (const FProperty* DisplayOptionsProperty = FLoudnessMeterRackUnitSettings::StaticStruct()->FindPropertyByName(LoudnessMetric.Name))
@@ -271,6 +289,7 @@ namespace AudioWidgets
 		}
 
 		LoudnessMeterWidgetView.InitTimerPanel(TimerPanelParams);
+		LoudnessMeterWidgetView.InitLoudnessScale(LoudnessScaleParams);
 		for (const FLoudnessMeterWidgetView::FLoudnessMetric& LoudnessMetric : LoudnessMetrics)
 		{
 			LoudnessMeterWidgetView.AddLoudnessMetric(LoudnessMetric);
diff --git a/Engine/Plugins/Runtime/AudioWidgets/Source/AudioWidgets/Private/LoudnessMeterWidgetView.cpp b/Engine/Plugins/Runtime/AudioWidgets/Source/AudioWidgets/Private/LoudnessMeterWidgetView.cpp
--- a/Engine/Plugins/Runtime/AudioWidgets/Source/AudioWidgets/Private/LoudnessMeterWidgetView.cpp
+++ b/Engine/Plugins/Runtime/AudioWidgets/Source/AudioWidgets/Private/LoudnessMeterWidgetView.cpp
@@ -14,6 +14,7 @@
 #include "Widgets/Colors/SColorBlock.h"
 #include "Widgets/Colors/SColorPicker.h"
 #include "Widgets/Input/SButton.h"
+#include "Widgets/Input/SNumericEntryBox.h"
 #include "Widgets/Layout/SSeparator.h"
 #include "Widgets/Layout/SSpacer.h"
 #include "Widgets/Views/STileView.h"
@@ -26,6 +27,67 @@ namespace AudioWidgets
 {
 	namespace LoudnessMeterWidgetView_Private
 	{
+		/**
+		 * Derived meter widget type adds bindable attributes for scale Range and Offset.
+		 */
+		class SVariableRangeMeterWidget : public SAudioMeterWidget
+		{
+		public:
+			SLATE_BEGIN_ARGS(SVariableRangeMeterWidget)
+				: _MeterValueColor(FLinearColor::Green)
+				, _Range(60)
+				, _Offset(0)
+				{
+				}
+
+				SLATE_ATTRIBUTE(FSlateColor, MeterValueColor)
+				SLATE_ATTRIBUTE(TArray<FAudioMeterChannelInfo>, MeterChannelInfo)
+				SLATE_ATTRIBUTE(int, Range)
+				SLATE_ATTRIBUTE(int, Offset)
+
+			SLATE_END_ARGS()
+
+			void Construct(const SVariableRangeMeterWidget::FArguments& InArgs)
+			{
+				AudioMeterWidgetStyle->SetMeterSize({ 700.0, 25.0 }); // MaxLength, Width.
+				AudioMeterWidgetStyle->SetScaleHashHeight(4.0f); // This is the tick mark length.
+
+				Range = InArgs._Range;
+				Offset = InArgs._Offset;
+
+				SAudioMeterWidget::Construct(
+					SAudioMeterWidget::FArguments()
+						.Orientation(EOrientation::Orient_Vertical)
+						.BackgroundColor(FLinearColor::Transparent)
+						.MeterBackgroundColor(FAudioMeterDefaultColorWidgetStyle::GetDefault().MeterBackgroundColor)
+						.MeterPeakColor(FLinearColor::Transparent)
+						.MeterClippingColor(FLinearColor::Transparent)
+						.MeterScaleColor(FLinearColor::White.CopyWithNewOpacity(0.25f))
+						.MeterScaleLabelColor(FLinearColor::White.CopyWithNewOpacity(0.25f))
+						.MeterValueColor(InArgs._MeterValueColor)
+						.MeterChannelInfo(InArgs._MeterChannelInfo)
+						.Style(AudioMeterWidgetStyle.Get())
+					);
+			}
+
+			virtual int32 OnPaint(const FPaintArgs& Args, const FGeometry& AllottedGeometry, const FSlateRect& MyCullingRect, FSlateWindowElementList& OutDrawElements, int32 LayerId, const FWidgetStyle& InWidgetStyle, bool bParentEnabled) const override
+			{
+				// Update the style from our attribute bindings:
+				const float RangeValue = Range.Get();
+				const float OffsetValue = Offset.Get();
+				AudioMeterWidgetStyle->SetValueRangeDb({ OffsetValue - RangeValue, OffsetValue });
+
+				// Paint the widget using the updated style:
+				return SAudioMeterWidget::OnPaint(Args, AllottedGeometry, MyCullingRect, OutDrawElements, LayerId, InWidgetStyle, bParentEnabled);
+			}
+
+		private:
+			const TUniquePtr<FAudioMeterWidgetStyle> AudioMeterWidgetStyle = MakeUnique<FAudioMeterWidgetStyle>(FAudioMeterWidgetStyle::GetDefault());
+			TAttribute<int> Offset;
+			TAttribute<int> Range;
+		};
+
+
 		FText FormatNumericValueText(TOptional<float> Value)
 		{
 			if (!Value.IsSet())
@@ -83,25 +145,12 @@ namespace AudioWidgets
 				];
 		}
 
-		TSharedRef<ITableRow> MakeMeterValueTileWidget(TSharedRef<const FLoudnessMeterWidgetView::FLoudnessMetric> InItem, const TSharedRef<STableViewBase>& OwnerTable)
+		TSharedRef<ITableRow> MakeMeterValueTileWidget(TSharedRef<const FLoudnessMeterWidgetView::FLoudnessMetric> InItem, const TSharedRef<STableViewBase>& OwnerTable, FLoudnessMeterWidgetView::FLoudnessScaleParams LoudnessScaleParams)
 		{
 			using FLoudnessMetricRef = TSharedRef<const FLoudnessMeterWidgetView::FLoudnessMetric>;
 
-			const auto MakeAudioMeterWidgetStyle = []()
-				{
-					FAudioMeterWidgetStyle MeterStyle = FAudioMeterWidgetStyle::GetDefault();
-					MeterStyle.SetMeterSize({ 700.0, 25.0 }); // MaxLength, Width.
-					MeterStyle.SetValueRangeDb({ -60, 0 });
-					MeterStyle.SetScaleHashHeight(4.0f); // This is the tick mark length.
-					return MeterStyle;
-				};
-
-			static const FAudioMeterWidgetStyle AudioMeterWidgetStyle = MakeAudioMeterWidgetStyle();
-
 			static const FSlateRoundedBoxBrush RoundedPanelBrush(FStyleColors::Panel, 4.0f);
 
-			const FAudioMeterDefaultColorWidgetStyle& DefaultColorStyle = FAudioMeterDefaultColorWidgetStyle::GetDefault();
-
 			return SNew(STableRow<FLoudnessMetricRef>, OwnerTable)
 				[
 					SNew(SBox)
@@ -126,22 +175,16 @@ namespace AudioWidgets
 										+ SVerticalBox::Slot()
 										.HAlign(EHorizontalAlignment::HAlign_Center)
 										[
-											SNew(SAudioMeterWidget)
-												.Orientation(EOrientation::Orient_Vertical)
-												.BackgroundColor(FLinearColor::Transparent)
-												.MeterBackgroundColor(DefaultColorStyle.MeterBackgroundColor)
+											SNew(SVariableRangeMeterWidget)
 												.MeterValueColor_Lambda([InItem]() { return InItem->Color.Get(); })
-												.MeterPeakColor(FLinearColor::Transparent)
-												.MeterClippingColor(FLinearColor::Transparent)
-												.MeterScaleColor(FLinearColor::White.CopyWithNewOpacity(0.25f))
-												.MeterScaleLabelColor(FLinearColor::White.CopyWithNewOpacity(0.25f))
-												.Style(&AudioMeterWidgetStyle)
 												.MeterChannelInfo_Lambda([InItem]() -> TArray<FAudioMeterChannelInfo>
 													{
 														constexpr float DefaultInvalidValue = MIN_VOLUME_DECIBELS;
 														const TOptional<float> LoudnessMetricValue = InItem->Value.Get();
 														return { { LoudnessMetricValue.Get(DefaultInvalidValue), DefaultInvalidValue, DefaultInvalidValue } };
 													})
+												.Range(LoudnessScaleParams.Range)
+												.Offset(LoudnessScaleParams.Offset)
 										]
 										+ SVerticalBox::Slot()
 										.AutoHeight()
@@ -185,6 +228,11 @@ namespace AudioWidgets
 		TimerPanelParams = InTimerPanelParams;
 	}
 
+	void FLoudnessMeterWidgetView::InitLoudnessScale(const FLoudnessScaleParams& InLoudnessScaleParams)
+	{
+		LoudnessScaleParams = InLoudnessScaleParams;
+	}
+
 	TSharedRef<SWidget> FLoudnessMeterWidgetView::MakeWidget() const
 	{
 		using SLoudnessMetricTileView = STileView<FLoudnessMetricRef>;
@@ -195,11 +243,49 @@ namespace AudioWidgets
 		const TSharedRef<FUICommandList> CommandList = MakeShared<FUICommandList>();
 		FSlimHorizontalToolBarBuilder ToolBarBuilder(CommandList, FMultiBoxCustomization::None, nullptr, bForceSmallIcons);
 
-		const auto MakeOptionsMenu = [CommandList, TimerPanelParams = this->TimerPanelParams, LoudnessMetrics = this->LoudnessMetrics]() -> TSharedRef<SWidget>
+		const auto MakeOptionsMenu = [
+			CommandList,
+			TimerPanelParams = this->TimerPanelParams,
+			LoudnessScaleParams = this->LoudnessScaleParams,
+			LoudnessMetrics = this->LoudnessMetrics]() -> TSharedRef<SWidget>
 			{
 				constexpr bool bShouldCloseWindowAfterMenuSelection = true;
 				FMenuBuilder MenuBuilder(bShouldCloseWindowAfterMenuSelection, CommandList);
 
+				{
+					MenuBuilder.BeginSection(NAME_None, LOCTEXT("LoudnessScale", "Loudness Scale"));
+
+					if (LoudnessScaleParams.OnRangeValueCommitted.IsBound())
+					{
+						TSharedRef RangeWidget = SNew(SNumericEntryBox<int>)
+							.MinDesiredValueWidth(100.0f)
+							.MinSliderValue(10)
+							.MaxSliderValue(70)
+							.AllowSpin(true)
+							.Value_Lambda([Range = LoudnessScaleParams.Range]() { return Range.Get(); })
+							.OnValueChanged(LoudnessScaleParams.OnRangeValueChanged)
+							.OnValueCommitted(LoudnessScaleParams.OnRangeValueCommitted);
+
+						MenuBuilder.AddWidget(RangeWidget, LOCTEXT("LoudnessScale_Range", "Range"));
+					}
+
+					if (LoudnessScaleParams.OnOffsetValueCommitted.IsBound())
+					{
+						TSharedRef OffsetWidget = SNew(SNumericEntryBox<int>)
+							.MinDesiredValueWidth(100.0f)
+							.MinSliderValue(-10)
+							.MaxSliderValue(30)
+							.AllowSpin(true)
+							.Value_Lambda([Offset = LoudnessScaleParams.Offset]() { return Offset.Get(); })
+							.OnValueChanged(LoudnessScaleParams.OnOffsetValueChanged)
+							.OnValueCommitted(LoudnessScaleParams.OnOffsetValueCommitted);
+
+						MenuBuilder.AddWidget(OffsetWidget, LOCTEXT("LoudnessScale_Offset", "Offset"));
+					}
+
+					MenuBuilder.EndSection();
+				}
+
 				{
 					MenuBuilder.BeginSection(NAME_None, LOCTEXT("LayoutOptions", "Layout Options"));
 
@@ -309,7 +395,7 @@ namespace AudioWidgets
 		constexpr float MeterMinHeight = 160.0f;
 
 		TSharedRef<SLoudnessMetricTileView> MeterTileView = SNew(SLoudnessMetricTileView)
-			.OnGenerateTile_Static(&LoudnessMeterWidgetView_Private::MakeMeterValueTileWidget)
+			.OnGenerateTile_Static(&LoudnessMeterWidgetView_Private::MakeMeterValueTileWidget, LoudnessScaleParams)
 			.ListItemsSource(MeterValues)
 			.ItemWidth(MeterWidth)
 			.ItemAlignment(EListItemAlignment::CenterAligned)
@@ -322,7 +408,8 @@ namespace AudioWidgets
 				const FVector2f PanelSize = TileView.GetTickSpaceGeometry().GetLocalSize();
 				const int32 NumItemsPerLine = FMath::Max(1, FMath::FloorToInt(PanelSize.X / MeterWidth));
 				const int32 NumLines = (NumItems + NumItemsPerLine - 1) / NumItemsPerLine;
-				return FMath::Max(PanelSize.Y / NumLines, MeterMinHeight);
+				const int32 HeightDivisor = FMath::Max(1, NumLines);
+				return FMath::Max(PanelSize.Y / HeightDivisor, MeterMinHeight);
 			}));
 
 		return SNew(SBorder)
diff --git a/Engine/Plugins/Runtime/AudioWidgets/Source/AudioWidgets/Private/LoudnessMeterWidgetView.h b/Engine/Plugins/Runtime/AudioWidgets/Source/AudioWidgets/Private/LoudnessMeterWidgetView.h
--- a/Engine/Plugins/Runtime/AudioWidgets/Source/AudioWidgets/Private/LoudnessMeterWidgetView.h
+++ b/Engine/Plugins/Runtime/AudioWidgets/Source/AudioWidgets/Private/LoudnessMeterWidgetView.h
@@ -41,10 +41,25 @@ namespace AudioWidgets
 			FSimpleDelegate OnVisibilityToggleRequested;
 		};
 
+		DECLARE_DELEGATE_OneParam(FOnIntValueChanged, int /*NewValue*/);
+		DECLARE_DELEGATE_TwoParams(FOnIntValueCommitted, int /*NewValue*/, ETextCommit::Type /*CommitType*/);
+
+		struct FLoudnessScaleParams
+		{
+			TAttribute<int> Range = 60;
+			FOnIntValueChanged OnRangeValueChanged;
+			FOnIntValueCommitted OnRangeValueCommitted;
+
+			TAttribute<int> Offset = 0;
+			FOnIntValueChanged OnOffsetValueChanged;
+			FOnIntValueCommitted OnOffsetValueCommitted;
+		};
+
 		FLoudnessMeterWidgetView();
 		~FLoudnessMeterWidgetView();
 
 		void InitTimerPanel(const FTimerPanelParams& InTimerPanelParams);
+		void InitLoudnessScale(const FLoudnessScaleParams& InLoudnessScaleParams);
 
 		TSharedRef<SWidget> MakeWidget() const;
 
@@ -59,5 +74,6 @@ namespace AudioWidgets
 		TSharedRef<FLoudnessMetricRefArray> NumericValues;
 		TSharedRef<FLoudnessMetricRefArray> MeterValues;
 		FTimerPanelParams TimerPanelParams;
+		FLoudnessScaleParams LoudnessScaleParams;
 	};
 } // namespace AudioWidgets
diff --git a/Engine/Plugins/Runtime/AudioWidgets/Source/AudioWidgets/Public/LoudnessMeterRackUnitSettings.h b/Engine/Plugins/Runtime/AudioWidgets/Source/AudioWidgets/Public/LoudnessMeterRackUnitSettings.h
--- a/Engine/Plugins/Runtime/AudioWidgets/Source/AudioWidgets/Public/LoudnessMeterRackUnitSettings.h
+++ b/Engine/Plugins/Runtime/AudioWidgets/Source/AudioWidgets/Public/LoudnessMeterRackUnitSettings.h
@@ -24,6 +24,12 @@ struct FLoudnessMeterRackUnitSettings
 {
 	GENERATED_BODY()
 
+	UPROPERTY(EditAnywhere, config, Category = LoudnessMeter, meta = (ClampMin = "10", ClampMax = "70", UIMin = "10", UIMax = "70"))
+	int LoudnessScaleRange = 60;
+
+	UPROPERTY(EditAnywhere, config, Category = LoudnessMeter, meta = (ClampMin = "-10", ClampMax = "30", UIMin = "-10", UIMax = "30"))
+	int LoudnessScaleOffset = 0;
+
 	UPROPERTY(EditAnywhere, config, Category = LoudnessMeter)
 	bool bDisplayAnalysisTimer = true;
 
diff --git a/Engine/Plugins/Runtime/AudioWidgets/Source/AudioWidgetsCore/Private/SAudioMeterWidget.cpp b/Engine/Plugins/Runtime/AudioWidgets/Source/AudioWidgetsCore/Private/SAudioMeterWidget.cpp
--- a/Engine/Plugins/Runtime/AudioWidgets/Source/AudioWidgetsCore/Private/SAudioMeterWidget.cpp
+++ b/Engine/Plugins/Runtime/AudioWidgets/Source/AudioWidgetsCore/Private/SAudioMeterWidget.cpp
@@ -50,10 +50,6 @@ int32 SAudioMeterWidget::OnPaint(const FPaintArgs& Args, const FGeometry& Allott
  	int32 MinValueDb = FMath::Min(static_cast<int32>(Style->ValueRangeDb.X), static_cast<int32>(Style->ValueRangeDb.Y));
  	int32 MaxValueDb = FMath::Max(static_cast<int32>(Style->ValueRangeDb.Y), static_cast<int32>(Style->ValueRangeDb.Y));
  
- 	// Snap the min/max values to the nearest hash 
- 	MaxValueDb -= MaxValueDb % DecibelsPerHash;
- 	MinValueDb -= MinValueDb % DecibelsPerHash;
- 
  	const float ValueDeltaDb = MaxValueDb - MinValueDb;
 
 	// Draw the background image
@@ -167,7 +163,8 @@ int32 SAudioMeterWidget::OnPaint(const FPaintArgs& Args, const FGeometry& Allott
 
 		const int32 ValueDelta = MaxValueDb - MinValueDb;
 
-		int32 CurrentHashValue = MaxValueDb;
+		// Round down MaxValueDb to find the first hash mark to draw:
+		int32 CurrentHashValue = DecibelsPerHash * FMath::FloorToInt32(float(MaxValueDb) / float(DecibelsPerHash));
 		while (CurrentHashValue >= MinValueDb)
 		{
 			// Get the fractional value for hash mark
