[Audio Widgets] Loudness Meter - Add menu option for displaying loudness values with target loudness relative scale.

diff --git a/Engine/Plugins/Runtime/AudioWidgets/Source/AudioWidgets/Private/AudioLoudnessMeter.cpp b/Engine/Plugins/Runtime/AudioWidgets/Source/AudioWidgets/Private/AudioLoudnessMeter.cpp
--- a/Engine/Plugins/Runtime/AudioWidgets/Source/AudioWidgets/Private/AudioLoudnessMeter.cpp
+++ b/Engine/Plugins/Runtime/AudioWidgets/Source/AudioWidgets/Private/AudioLoudnessMeter.cpp
@@ -221,6 +221,22 @@ namespace AudioWidgets
 						SettingsHelper.SaveConfig();
 					});
 
+				LoudnessScaleParams.Target.BindLambda([=]() { return SettingsHelper.GetRackUnitSettings().LoudnessScaleTarget; });
+				LoudnessScaleParams.OnTargetValueChanged.BindLambda([=](int Value) { SettingsHelper.GetRackUnitSettings().LoudnessScaleTarget = Value; });
+				LoudnessScaleParams.OnTargetValueCommitted.BindLambda([=](int Value, ETextCommit::Type CommitType)
+					{
+						SettingsHelper.GetRackUnitSettings().LoudnessScaleTarget = Value;
+						SettingsHelper.SaveConfig();
+					});
+
+				LoudnessScaleParams.bRelativeScale.BindLambda([=]() { return SettingsHelper.GetRackUnitSettings().bUseRelativeLoudnessScale; });
+				LoudnessScaleParams.OnRelativeScaleToggleRequested.BindLambda([=]()
+					{
+						FLoudnessMeterRackUnitSettings& LoudnessMeterSettings = SettingsHelper.GetRackUnitSettings();
+						LoudnessMeterSettings.bUseRelativeLoudnessScale = !LoudnessMeterSettings.bUseRelativeLoudnessScale;
+						SettingsHelper.SaveConfig();
+					});
+
 				for (FLoudnessMeterWidgetView::FLoudnessMetric& LoudnessMetric : LoudnessMetrics)
 				{
 					if (const FProperty* DisplayOptionsProperty = FLoudnessMeterRackUnitSettings::StaticStruct()->FindPropertyByName(LoudnessMetric.Name))
diff --git a/Engine/Plugins/Runtime/AudioWidgets/Source/AudioWidgets/Private/LoudnessMeterWidgetView.cpp b/Engine/Plugins/Runtime/AudioWidgets/Source/AudioWidgets/Private/LoudnessMeterWidgetView.cpp
--- a/Engine/Plugins/Runtime/AudioWidgets/Source/AudioWidgets/Private/LoudnessMeterWidgetView.cpp
+++ b/Engine/Plugins/Runtime/AudioWidgets/Source/AudioWidgets/Private/LoudnessMeterWidgetView.cpp
@@ -87,29 +87,63 @@ namespace AudioWidgets
 			TAttribute<int> Range;
 		};
 
-
-		FText FormatNumericValueText(TOptional<float> Value)
+		/**
+		 * If the Relative Scale attribute is true this will convert values given in LUFS into values in LU, relative to the Target attribute.
+		 * Otherwise it will leave the LUFS values as-is.
+		 */
+		class FLoudnessValueConverter
 		{
-			if (!Value.IsSet())
+		public:
+			FLoudnessValueConverter(const FLoudnessMeterWidgetView::FLoudnessScaleParams& LoudnessScaleParams)
+				: bRelativeScale(LoudnessScaleParams.bRelativeScale)
+				, Target(LoudnessScaleParams.Target)
 			{
-				return FText::GetEmpty();
 			}
-			else if (*Value == -std::numeric_limits<float>::infinity())
+
+			float ConvertValue(float Value) const
 			{
-				return LOCTEXT("NegativeInfinity", "-∞");
+				return (bRelativeScale.Get()) ? (Value - Target.Get()) : Value;
 			}
 
-			static const FNumberFormattingOptions NumberFormattingOptions = FNumberFormattingOptions()
-				.SetMinimumFractionalDigits(1)
-				.SetMaximumFractionalDigits(1);
+			TOptional<float> ConvertValue(TOptional<float> Value) const
+			{
+				if (Value.IsSet())
+				{
+					return ConvertValue(*Value);
+				}
+				return NullOpt;
+			}
 
-			return FText::AsNumber(*Value, &NumberFormattingOptions);
-		}
+			FText FormatNumericValueText(TOptional<float> Value) const
+			{
+				if (!Value.IsSet())
+				{
+					return FText::GetEmpty();
+				}
+				else if (*Value == -std::numeric_limits<float>::infinity())
+				{
+					return LOCTEXT("NegativeInfinity", "-∞");
+				}
+
+				static const FNumberFormattingOptions NumberFormattingOptions = FNumberFormattingOptions()
+					.SetMinimumFractionalDigits(1)
+					.SetMaximumFractionalDigits(1);
+
+				const float ConvertedValue = ConvertValue(*Value);
+				return FText::AsNumber(ConvertedValue, &NumberFormattingOptions);
+			}
+
+		private:
+			TAttribute<bool> bRelativeScale;
+			TAttribute<int> Target;
+		};
 
-		TSharedRef<ITableRow> MakeNumericValueTileWidget(TSharedRef<const FLoudnessMeterWidgetView::FLoudnessMetric> InItem, const TSharedRef<STableViewBase>& OwnerTable)
+		TSharedRef<ITableRow> MakeNumericValueTileWidget(TSharedRef<const FLoudnessMeterWidgetView::FLoudnessMetric> InItem, const TSharedRef<STableViewBase>& OwnerTable, FLoudnessMeterWidgetView::FLoudnessScaleParams LoudnessScaleParams)
 		{
 			using FLoudnessMetricRef = TSharedRef<const FLoudnessMeterWidgetView::FLoudnessMetric>;
 
+			const FLoudnessValueConverter LoudnessValueConverter(LoudnessScaleParams);
+
 			return SNew(STableRow<FLoudnessMetricRef>, OwnerTable)
 				[
 					SNew(SVerticalBox)
@@ -125,7 +159,7 @@ namespace AudioWidgets
 						.HAlign(EHorizontalAlignment::HAlign_Center)
 						[
 							SNew(STextBlock)
-								.Text_Lambda([InItem]() { return FormatNumericValueText(InItem->Value.Get()); })
+								.Text_Lambda([InItem, LoudnessValueConverter]() { return LoudnessValueConverter.FormatNumericValueText(InItem->Value.Get()); })
 								.Font(FStyleDefaults::GetFontInfo(14))
 								.ColorAndOpacity_Lambda([InItem]() { return InItem->Color.Get().CopyWithNewOpacity(1.0f); })
 								.SimpleTextMode(true)
@@ -151,6 +185,8 @@ namespace AudioWidgets
 
 			static const FSlateRoundedBoxBrush RoundedPanelBrush(FStyleColors::Panel, 4.0f);
 
+			const FLoudnessValueConverter LoudnessValueConverter(LoudnessScaleParams);
+
 			return SNew(STableRow<FLoudnessMetricRef>, OwnerTable)
 				[
 					SNew(SBox)
@@ -167,7 +203,7 @@ namespace AudioWidgets
 										.HAlign(EHorizontalAlignment::HAlign_Center)
 										[
 											SNew(STextBlock)
-												.Text_Lambda([InItem]() { return FormatNumericValueText(InItem->Value.Get()); })
+												.Text_Lambda([InItem, LoudnessValueConverter]() { return LoudnessValueConverter.FormatNumericValueText(InItem->Value.Get()); })
 												.Font(FStyleDefaults::GetFontInfo(14))
 												.ColorAndOpacity_Lambda([InItem]() { return InItem->Color.Get().CopyWithNewOpacity(1.0f); })
 												.SimpleTextMode(true)
@@ -177,14 +213,14 @@ namespace AudioWidgets
 										[
 											SNew(SVariableRangeMeterWidget)
 												.MeterValueColor_Lambda([InItem]() { return InItem->Color.Get(); })
-												.MeterChannelInfo_Lambda([InItem]() -> TArray<FAudioMeterChannelInfo>
+												.MeterChannelInfo_Lambda([InItem, LoudnessValueConverter]() -> TArray<FAudioMeterChannelInfo>
 													{
 														constexpr float DefaultInvalidValue = MIN_VOLUME_DECIBELS;
-														const TOptional<float> LoudnessMetricValue = InItem->Value.Get();
-														return { { LoudnessMetricValue.Get(DefaultInvalidValue), DefaultInvalidValue, DefaultInvalidValue } };
+														const TOptional<float> ConvertedValue = LoudnessValueConverter.ConvertValue(InItem->Value.Get());
+														return { { ConvertedValue.Get(DefaultInvalidValue), DefaultInvalidValue, DefaultInvalidValue } };
 													})
 												.Range(LoudnessScaleParams.Range)
-												.Offset(LoudnessScaleParams.Offset)
+												.Offset_Lambda([Offset = LoudnessScaleParams.Offset, LoudnessValueConverter] { return LoudnessValueConverter.ConvertValue(Offset.Get()); })
 										]
 										+ SVerticalBox::Slot()
 										.AutoHeight()
@@ -283,6 +319,35 @@ namespace AudioWidgets
 						MenuBuilder.AddWidget(OffsetWidget, LOCTEXT("LoudnessScale_Offset", "Offset"));
 					}
 
+					if (LoudnessScaleParams.OnTargetValueCommitted.IsBound())
+					{
+						TSharedRef TargetWidget = SNew(SNumericEntryBox<int>)
+							.MinDesiredValueWidth(100.0f)
+							.MinSliderValue(-60)
+							.MaxSliderValue(30)
+							.AllowSpin(true)
+							.Value_Lambda([Target = LoudnessScaleParams.Target]() { return Target.Get(); })
+							.OnValueChanged(LoudnessScaleParams.OnTargetValueChanged)
+							.OnValueCommitted(LoudnessScaleParams.OnTargetValueCommitted);
+
+						MenuBuilder.AddWidget(TargetWidget, LOCTEXT("LoudnessScale_Target", "Target"));
+					}
+
+					if (LoudnessScaleParams.OnRelativeScaleToggleRequested.IsBound())
+					{
+						MenuBuilder.AddMenuEntry(
+							LOCTEXT("LoudnessScale_RelativeScale_Label", "Relative Scale"),
+							LOCTEXT("LoudnessScale_RelativeScale_ToolTip", "Loudness values should be displayed relative to the Target loudness"),
+							FSlateIcon(),
+							FUIAction(
+								LoudnessScaleParams.OnRelativeScaleToggleRequested,
+								FCanExecuteAction(),
+								FIsActionChecked::CreateLambda([bRelativeScale = LoudnessScaleParams.bRelativeScale]() { return bRelativeScale.Get(); })
+							),
+							NAME_None,
+							EUserInterfaceActionType::ToggleButton);
+					}
+
 					MenuBuilder.EndSection();
 				}
 
@@ -483,7 +548,7 @@ namespace AudioWidgets
 					[
 						SNew(SLoudnessMetricTileView)
 							.Visibility(GetNumericValuesPanelVisibility)
-							.OnGenerateTile_Static(&LoudnessMeterWidgetView_Private::MakeNumericValueTileWidget)
+							.OnGenerateTile_Static(&LoudnessMeterWidgetView_Private::MakeNumericValueTileWidget, LoudnessScaleParams)
 							.ListItemsSource(NumericValues)
 							.ItemHeight(60.0f)
 							.ItemWidth(60.0f)
diff --git a/Engine/Plugins/Runtime/AudioWidgets/Source/AudioWidgets/Private/LoudnessMeterWidgetView.h b/Engine/Plugins/Runtime/AudioWidgets/Source/AudioWidgets/Private/LoudnessMeterWidgetView.h
--- a/Engine/Plugins/Runtime/AudioWidgets/Source/AudioWidgets/Private/LoudnessMeterWidgetView.h
+++ b/Engine/Plugins/Runtime/AudioWidgets/Source/AudioWidgets/Private/LoudnessMeterWidgetView.h
@@ -53,6 +53,13 @@ namespace AudioWidgets
 			TAttribute<int> Offset = 0;
 			FOnIntValueChanged OnOffsetValueChanged;
 			FOnIntValueCommitted OnOffsetValueCommitted;
+
+			TAttribute<int> Target = -23;
+			FOnIntValueChanged OnTargetValueChanged;
+			FOnIntValueCommitted OnTargetValueCommitted;
+
+			TAttribute<bool> bRelativeScale = false;
+			FSimpleDelegate OnRelativeScaleToggleRequested;
 		};
 
 		FLoudnessMeterWidgetView();
diff --git a/Engine/Plugins/Runtime/AudioWidgets/Source/AudioWidgets/Public/LoudnessMeterRackUnitSettings.h b/Engine/Plugins/Runtime/AudioWidgets/Source/AudioWidgets/Public/LoudnessMeterRackUnitSettings.h
--- a/Engine/Plugins/Runtime/AudioWidgets/Source/AudioWidgets/Public/LoudnessMeterRackUnitSettings.h
+++ b/Engine/Plugins/Runtime/AudioWidgets/Source/AudioWidgets/Public/LoudnessMeterRackUnitSettings.h
@@ -30,6 +30,12 @@ struct FLoudnessMeterRackUnitSettings
 	UPROPERTY(EditAnywhere, config, Category = LoudnessMeter, meta = (ClampMin = "-10", ClampMax = "30", UIMin = "-10", UIMax = "30"))
 	int LoudnessScaleOffset = 0;
 
+	UPROPERTY(EditAnywhere, config, Category = LoudnessMeter, meta = (ClampMin = "-60", ClampMax = "30", UIMin = "-60", UIMax = "30"))
+	int LoudnessScaleTarget = -23;
+
+	UPROPERTY(EditAnywhere, config, Category = LoudnessMeter)
+	bool bUseRelativeLoudnessScale = false;
+
 	UPROPERTY(EditAnywhere, config, Category = LoudnessMeter)
 	bool bDisplayAnalysisTimer = true;
 
