[Audio Widgets] Loudness Meter - Display units (LUFS or LU).

diff --git a/Engine/Plugins/Runtime/AudioWidgets/Source/AudioWidgets/Private/LoudnessMeterWidgetView.cpp b/Engine/Plugins/Runtime/AudioWidgets/Source/AudioWidgets/Private/LoudnessMeterWidgetView.cpp
--- a/Engine/Plugins/Runtime/AudioWidgets/Source/AudioWidgets/Private/LoudnessMeterWidgetView.cpp
+++ b/Engine/Plugins/Runtime/AudioWidgets/Source/AudioWidgets/Private/LoudnessMeterWidgetView.cpp
@@ -138,12 +138,41 @@ namespace AudioWidgets
 			TAttribute<int> Target;
 		};
 
+		TSharedRef<SWidget> MakeValueReadoutWidget(TSharedRef<const FLoudnessMeterWidgetView::FLoudnessMetric> InItem, const FLoudnessMeterWidgetView::FLoudnessScaleParams& LoudnessScaleParams)
+		{
+			const FLoudnessValueConverter LoudnessValueConverter(LoudnessScaleParams);
+
+			const auto GetValueColorAndOpacity = [InItem]() { return InItem->Color.Get().CopyWithNewOpacity(1.0f); };
+
+			return SNew(SOverlay)
+				+ SOverlay::Slot()
+				[
+					SNew(SSpacer)
+						.Size(FVector2D(46.0f, 28.0f))
+				]
+				+ SOverlay::Slot()
+				[
+					SNew(STextBlock)
+						.Text_Lambda([InItem, LoudnessValueConverter]() { return LoudnessValueConverter.FormatNumericValueText(InItem->Value.Get()); })
+						.Font(FStyleDefaults::GetFontInfo(14))
+						.ColorAndOpacity_Lambda(GetValueColorAndOpacity)
+						.Justification(ETextJustify::Right)
+				]
+				+ SOverlay::Slot()
+				.VAlign(EVerticalAlignment::VAlign_Bottom)
+				[
+					SNew(STextBlock)
+						.Text_Lambda([bRelativeScale = LoudnessScaleParams.bRelativeScale]() { return (bRelativeScale.Get()) ? INVTEXT("LU") : INVTEXT("LUFS"); })
+						.Font(FStyleDefaults::GetFontInfo(6))
+						.ColorAndOpacity_Lambda(GetValueColorAndOpacity)
+						.Justification(ETextJustify::Right)
+				];
+		}
+
 		TSharedRef<ITableRow> MakeNumericValueTileWidget(TSharedRef<const FLoudnessMeterWidgetView::FLoudnessMetric> InItem, const TSharedRef<STableViewBase>& OwnerTable, FLoudnessMeterWidgetView::FLoudnessScaleParams LoudnessScaleParams)
 		{
 			using FLoudnessMetricRef = TSharedRef<const FLoudnessMeterWidgetView::FLoudnessMetric>;
 
-			const FLoudnessValueConverter LoudnessValueConverter(LoudnessScaleParams);
-
 			return SNew(STableRow<FLoudnessMetricRef>, OwnerTable)
 				[
 					SNew(SVerticalBox)
@@ -154,19 +183,14 @@ namespace AudioWidgets
 							SNew(SSpacer)
 						]
 						+ SVerticalBox::Slot()
-						.MinHeight(20)
-						.MaxHeight(20)
+						.AutoHeight()
 						.HAlign(EHorizontalAlignment::HAlign_Center)
 						[
-							SNew(STextBlock)
-								.Text_Lambda([InItem, LoudnessValueConverter]() { return LoudnessValueConverter.FormatNumericValueText(InItem->Value.Get()); })
-								.Font(FStyleDefaults::GetFontInfo(14))
-								.ColorAndOpacity_Lambda([InItem]() { return InItem->Color.Get().CopyWithNewOpacity(1.0f); })
-								.SimpleTextMode(true)
+							MakeValueReadoutWidget(InItem, LoudnessScaleParams)
 						]
 						+ SVerticalBox::Slot()
-						.MinHeight(32)
-						.MaxHeight(32)
+						.MinHeight(24)
+						.MaxHeight(24)
 						.HAlign(EHorizontalAlignment::HAlign_Center)
 						.VAlign(EVerticalAlignment::VAlign_Center)
 						[
@@ -202,11 +226,7 @@ namespace AudioWidgets
 										.AutoHeight()
 										.HAlign(EHorizontalAlignment::HAlign_Center)
 										[
-											SNew(STextBlock)
-												.Text_Lambda([InItem, LoudnessValueConverter]() { return LoudnessValueConverter.FormatNumericValueText(InItem->Value.Get()); })
-												.Font(FStyleDefaults::GetFontInfo(14))
-												.ColorAndOpacity_Lambda([InItem]() { return InItem->Color.Get().CopyWithNewOpacity(1.0f); })
-												.SimpleTextMode(true)
+											MakeValueReadoutWidget(InItem, LoudnessScaleParams)
 										]
 										+ SVerticalBox::Slot()
 										.HAlign(EHorizontalAlignment::HAlign_Center)
