[Audio Insights] Include heavier IsValid check on Analyzer in AudioMeter before trying to shutdown

diff --git a/Engine/Plugins/Runtime/AudioWidgets/Source/AudioWidgets/Private/AudioMeter.cpp b/Engine/Plugins/Runtime/AudioWidgets/Source/AudioWidgets/Private/AudioMeter.cpp
--- a/Engine/Plugins/Runtime/AudioWidgets/Source/AudioWidgets/Private/AudioMeter.cpp
+++ b/Engine/Plugins/Runtime/AudioWidgets/Source/AudioWidgets/Private/AudioMeter.cpp
@@ -293,7 +293,7 @@ namespace AudioWidgets
 
 	void FAudioMeter::Teardown()
 	{
-		if (Analyzer.IsValid() && Analyzer->IsValidLowLevel())
+		if (Analyzer.IsValid() && IsValid(Analyzer.Get()) && Analyzer->IsValidLowLevel())
 		{
 			Analyzer->StopAnalyzing();
 			if (ResultsDelegateHandle.IsValid())
