[Audio Insights] Shutdown crash fix in AudioMeter when using command line arg -stompmalloc. Ensure the engine is still alive before trying to dereference Analyzer.

diff --git a/Engine/Plugins/Runtime/AudioWidgets/Source/AudioWidgets/Private/AudioMeter.cpp b/Engine/Plugins/Runtime/AudioWidgets/Source/AudioWidgets/Private/AudioMeter.cpp
--- a/Engine/Plugins/Runtime/AudioWidgets/Source/AudioWidgets/Private/AudioMeter.cpp
+++ b/Engine/Plugins/Runtime/AudioWidgets/Source/AudioWidgets/Private/AudioMeter.cpp
@@ -6,6 +6,7 @@
 #include "AudioMaterialSlate/SAudioMaterialMeter.h"
 #include "AudioMixerDevice.h"
 #include "AudioWidgetsStyle.h"
+#include "Engine/Engine.h"
 #include "SAudioMeter.h"
 #include "UObject/UObjectGlobals.h"
 #include "Widgets/Docking/SDockTab.h"
@@ -293,17 +294,16 @@ namespace AudioWidgets
 
 	void FAudioMeter::Teardown()
 	{
-		if (Analyzer.IsValid() && IsValid(Analyzer.Get()) && Analyzer->IsValidLowLevel())
+		if (GEngine && Analyzer.IsValid() && Analyzer->IsValidLowLevel())
 		{
 			Analyzer->StopAnalyzing();
 			if (ResultsDelegateHandle.IsValid())
 			{
 				Analyzer->OnLatestPerChannelMeterResultsNative.Remove(ResultsDelegateHandle);
 			}
-
-			Analyzer.Reset();
 		}
 
+		Analyzer.Reset();
 		ResultsDelegateHandle.Reset();
 		AudioBus.Reset();
 		ChannelInfo.Reset();
