[Audio Insights] Shutdown crash fix when using command line -stompmalloc inside Loudness Meter, Spectral Analyzer and Spectrogram. Update AudioMeter too to use more specific validity check on UObject system (as opposed to engine validity).

diff --git a/Engine/Plugins/Runtime/AudioWidgets/Source/AudioWidgets/Private/AudioLoudnessMeter.cpp b/Engine/Plugins/Runtime/AudioWidgets/Source/AudioWidgets/Private/AudioLoudnessMeter.cpp
--- a/Engine/Plugins/Runtime/AudioWidgets/Source/AudioWidgets/Private/AudioLoudnessMeter.cpp
+++ b/Engine/Plugins/Runtime/AudioWidgets/Source/AudioWidgets/Private/AudioLoudnessMeter.cpp
@@ -120,12 +120,15 @@ namespace AudioWidgets
 
 	void FAudioLoudnessMeter::Teardown()
 	{
-		if (LKFSAnalyzer.IsValid() && LKFSAnalyzer->IsValidLowLevel())
+		if (UObjectInitialized() && LKFSAnalyzer.IsValid() && LKFSAnalyzer->IsValidLowLevel())
 		{
 			LKFSAnalyzer->StopAnalyzing();
 			ReleaseLKFSAnalyzer();
 		}
 
+		LatestOverallLKFSResultsDelegateHandle.Reset();
+		LKFSAnalyzer.Reset();
+
 		AudioBus.Reset();
 		LKFSSettings.Reset();
 		bUseExternalAudioBus = false;
diff --git a/Engine/Plugins/Runtime/AudioWidgets/Source/AudioWidgets/Private/AudioMeter.cpp b/Engine/Plugins/Runtime/AudioWidgets/Source/AudioWidgets/Private/AudioMeter.cpp
--- a/Engine/Plugins/Runtime/AudioWidgets/Source/AudioWidgets/Private/AudioMeter.cpp
+++ b/Engine/Plugins/Runtime/AudioWidgets/Source/AudioWidgets/Private/AudioMeter.cpp
@@ -294,7 +294,7 @@ namespace AudioWidgets
 
 	void FAudioMeter::Teardown()
 	{
-		if (GEngine && Analyzer.IsValid() && Analyzer->IsValidLowLevel())
+		if (UObjectInitialized() && Analyzer.IsValid() && Analyzer->IsValidLowLevel())
 		{
 			Analyzer->StopAnalyzing();
 			if (ResultsDelegateHandle.IsValid())
diff --git a/Engine/Plugins/Runtime/AudioWidgets/Source/AudioWidgets/Private/AudioSpectrogram.cpp b/Engine/Plugins/Runtime/AudioWidgets/Source/AudioWidgets/Private/AudioSpectrogram.cpp
--- a/Engine/Plugins/Runtime/AudioWidgets/Source/AudioWidgets/Private/AudioSpectrogram.cpp
+++ b/Engine/Plugins/Runtime/AudioWidgets/Source/AudioWidgets/Private/AudioSpectrogram.cpp
@@ -218,25 +218,34 @@ namespace AudioWidgets
 
 	void FAudioSpectrogram::Teardown()
 	{
-		if (SpectrumAnalyzer.IsValid() && SpectrumAnalyzer->IsValidLowLevel())
+		if (UObjectInitialized())
 		{
-			if (ActiveAnalyzerType == EAudioSpectrumAnalyzerType::FFT)
+			if (SpectrumAnalyzer.IsValid() && SpectrumAnalyzer->IsValidLowLevel())
 			{
-				SpectrumAnalyzer->StopAnalyzing();
+				if (ActiveAnalyzerType == EAudioSpectrumAnalyzerType::FFT)
+				{
+					SpectrumAnalyzer->StopAnalyzing();
+				}
+
+				ReleaseSynesthesiaSpectrumAnalyzer();
 			}
+			
+			if (ConstantQAnalyzer.IsValid() && ConstantQAnalyzer->IsValidLowLevel())
+			{
+				if (ActiveAnalyzerType == EAudioSpectrumAnalyzerType::CQT)
+				{
+					ConstantQAnalyzer->StopAnalyzing();
+				}
 
-			ReleaseSynesthesiaSpectrumAnalyzer();
+				ReleaseConstantQAnalyzer();
+			}
 		}
 
-		if (ConstantQAnalyzer.IsValid() && ConstantQAnalyzer->IsValidLowLevel())
-		{
-			if (ActiveAnalyzerType == EAudioSpectrumAnalyzerType::CQT)
-			{
-				ConstantQAnalyzer->StopAnalyzing();
-			}
+		SpectrumResultsDelegateHandle.Reset();
+		SpectrumAnalyzer.Reset();
 
-			ReleaseConstantQAnalyzer();
-		}
+		ConstantQResultsDelegateHandle.Reset();
+		ConstantQAnalyzer.Reset();
 
 		ActiveAnalyzerType.Reset();
 
diff --git a/Engine/Plugins/Runtime/AudioWidgets/Source/AudioWidgets/Private/AudioSpectrumAnalyzer.cpp b/Engine/Plugins/Runtime/AudioWidgets/Source/AudioWidgets/Private/AudioSpectrumAnalyzer.cpp
--- a/Engine/Plugins/Runtime/AudioWidgets/Source/AudioWidgets/Private/AudioSpectrumAnalyzer.cpp
+++ b/Engine/Plugins/Runtime/AudioWidgets/Source/AudioWidgets/Private/AudioSpectrumAnalyzer.cpp
@@ -316,26 +316,35 @@ namespace AudioWidgets
 
 	void FAudioSpectrumAnalyzer::Teardown()
 	{
-		if (SpectrumAnalyzer.IsValid() && SpectrumAnalyzer->IsValidLowLevel())
+		if (UObjectInitialized())
 		{
-			if (ActiveAnalyzerType == EAudioSpectrumAnalyzerType::FFT)
+			if (SpectrumAnalyzer.IsValid() && SpectrumAnalyzer->IsValidLowLevel())
 			{
-				SpectrumAnalyzer->StopAnalyzing();
-			}
+				if (ActiveAnalyzerType == EAudioSpectrumAnalyzerType::FFT)
+				{
+					SpectrumAnalyzer->StopAnalyzing();
+				}
 
-			ReleaseSynesthesiaSpectrumAnalyzer();
-		}
+				ReleaseSynesthesiaSpectrumAnalyzer();
+			}
 
-		if (ConstantQAnalyzer.IsValid() && ConstantQAnalyzer->IsValidLowLevel())
-		{
-			if (ActiveAnalyzerType == EAudioSpectrumAnalyzerType::CQT)
+			if (ConstantQAnalyzer.IsValid() && ConstantQAnalyzer->IsValidLowLevel())
 			{
-				ConstantQAnalyzer->StopAnalyzing();
-			}
+				if (ActiveAnalyzerType == EAudioSpectrumAnalyzerType::CQT)
+				{
+					ConstantQAnalyzer->StopAnalyzing();
+				}
 
-			ReleaseConstantQAnalyzer();
+				ReleaseConstantQAnalyzer();
+			}
 		}
 
+		SpectrumResultsDelegateHandle.Reset();
+		SpectrumAnalyzer.Reset();
+
+		ConstantQResultsDelegateHandle.Reset();
+		ConstantQAnalyzer.Reset();
+
 		ActiveAnalyzerType.Reset();
 		PrevTimeStamp.Reset();
 		CenterFrequencies.Empty();
