[Audio Widgets] Loudness Meter - Order of tiles is controllable via settings params. This is a prelude to the CL(s) for drag/drop reordering of loudness meters.

#Synced-CL 48281405

diff --git a/Engine/Plugins/Runtime/AudioWidgets/Source/AudioWidgets/Private/AudioLoudnessMeter.cpp b/Engine/Plugins/Runtime/AudioWidgets/Source/AudioWidgets/Private/AudioLoudnessMeter.cpp
--- a/Engine/Plugins/Runtime/AudioWidgets/Source/AudioWidgets/Private/AudioLoudnessMeter.cpp
+++ b/Engine/Plugins/Runtime/AudioWidgets/Source/AudioWidgets/Private/AudioLoudnessMeter.cpp
@@ -167,6 +167,8 @@ namespace AudioWidgets
 		};
 
 		FLoudnessMeterWidgetView::FLoudnessScaleParams LoudnessScaleParams;
+		FLoudnessMeterWidgetView::FDisplayOrderParams ValuesDisplayOrderParams;
+		FLoudnessMeterWidgetView::FDisplayOrderParams MetersDisplayOrderParams;
 
 		FLoudnessMeterWidgetView::FLoudnessMetric LoudnessMetrics[] =
 		{
@@ -240,6 +242,20 @@ namespace AudioWidgets
 						SettingsHelper.SaveConfig();
 					});
 
+				ValuesDisplayOrderParams.PermutationIndex.BindLambda([=]() { return SettingsHelper.GetRackUnitSettings().ValuesOrderingPermutation; });
+				ValuesDisplayOrderParams.OnPermutationChanged.BindLambda([=](int NewPermutationIndex)
+					{
+						SettingsHelper.GetRackUnitSettings().ValuesOrderingPermutation = NewPermutationIndex;
+						SettingsHelper.SaveConfig();
+					});
+
+				MetersDisplayOrderParams.PermutationIndex.BindLambda([=]() { return SettingsHelper.GetRackUnitSettings().MetersOrderingPermutation; });
+				MetersDisplayOrderParams.OnPermutationChanged.BindLambda([=](int NewPermutationIndex)
+					{
+						SettingsHelper.GetRackUnitSettings().MetersOrderingPermutation = NewPermutationIndex;
+						SettingsHelper.SaveConfig();
+					});
+
 				for (FLoudnessMeterWidgetView::FLoudnessMetric& LoudnessMetric : LoudnessMetrics)
 				{
 					if (const FProperty* DisplayOptionsProperty = FLoudnessMeterRackUnitSettings::StaticStruct()->FindPropertyByName(LoudnessMetric.Name))
@@ -309,6 +325,8 @@ namespace AudioWidgets
 
 		LoudnessMeterWidgetView.InitTimerPanel(TimerPanelParams);
 		LoudnessMeterWidgetView.InitLoudnessScale(LoudnessScaleParams);
+		LoudnessMeterWidgetView.InitValuesDisplayOrder(ValuesDisplayOrderParams);
+		LoudnessMeterWidgetView.InitMetersDisplayOrder(MetersDisplayOrderParams);
 		for (const FLoudnessMeterWidgetView::FLoudnessMetric& LoudnessMetric : LoudnessMetrics)
 		{
 			LoudnessMeterWidgetView.AddLoudnessMetric(LoudnessMetric);
diff --git a/Engine/Plugins/Runtime/AudioWidgets/Source/AudioWidgets/Private/LoudnessMeterWidgetView.cpp b/Engine/Plugins/Runtime/AudioWidgets/Source/AudioWidgets/Private/LoudnessMeterWidgetView.cpp
--- a/Engine/Plugins/Runtime/AudioWidgets/Source/AudioWidgets/Private/LoudnessMeterWidgetView.cpp
+++ b/Engine/Plugins/Runtime/AudioWidgets/Source/AudioWidgets/Private/LoudnessMeterWidgetView.cpp
@@ -9,6 +9,7 @@
 #include "Brushes/SlateRoundedBoxBrush.h"
 #include "Framework/Application/SlateApplication.h"
 #include "Framework/MultiBox/MultiBoxBuilder.h"
+#include "Permutations.h"
 #include "SAudioMeterWidget.h"
 #include "Styling/StyleColors.h"
 #include "Widgets/Colors/SColorBlock.h"
@@ -27,6 +28,9 @@ namespace AudioWidgets
 {
 	namespace LoudnessMeterWidgetView_Private
 	{
+		using FLoudnessScaleParams = FLoudnessMeterWidgetView::FLoudnessScaleParams;
+		using FLoudnessMetricRef = TSharedRef<const FLoudnessMeterWidgetView::FLoudnessMetric>;
+
 		/**
 		 * Derived meter widget type adds bindable attributes for scale Range and Offset.
 		 */
@@ -94,7 +98,7 @@ namespace AudioWidgets
 		class FLoudnessValueConverter
 		{
 		public:
-			FLoudnessValueConverter(const FLoudnessMeterWidgetView::FLoudnessScaleParams& LoudnessScaleParams)
+			FLoudnessValueConverter(const FLoudnessScaleParams& LoudnessScaleParams)
 				: bRelativeScale(LoudnessScaleParams.bRelativeScale)
 				, Target(LoudnessScaleParams.Target)
 			{
@@ -138,7 +142,7 @@ namespace AudioWidgets
 			TAttribute<int> Target;
 		};
 
-		TSharedRef<SWidget> MakeValueReadoutWidget(TSharedRef<const FLoudnessMeterWidgetView::FLoudnessMetric> InItem, const FLoudnessMeterWidgetView::FLoudnessScaleParams& LoudnessScaleParams)
+		TSharedRef<SWidget> MakeValueReadoutWidget(FLoudnessMetricRef InItem, const FLoudnessScaleParams& LoudnessScaleParams)
 		{
 			const FLoudnessValueConverter LoudnessValueConverter(LoudnessScaleParams);
 
@@ -169,10 +173,8 @@ namespace AudioWidgets
 				];
 		}
 
-		TSharedRef<ITableRow> MakeNumericValueTileWidget(TSharedRef<const FLoudnessMeterWidgetView::FLoudnessMetric> InItem, const TSharedRef<STableViewBase>& OwnerTable, FLoudnessMeterWidgetView::FLoudnessScaleParams LoudnessScaleParams)
+		TSharedRef<ITableRow> MakeNumericValueTileWidget(FLoudnessMetricRef InItem, const TSharedRef<STableViewBase>& OwnerTable, FLoudnessScaleParams LoudnessScaleParams)
 		{
-			using FLoudnessMetricRef = TSharedRef<const FLoudnessMeterWidgetView::FLoudnessMetric>;
-
 			return SNew(STableRow<FLoudnessMetricRef>, OwnerTable)
 				[
 					SNew(SVerticalBox)
@@ -203,10 +205,8 @@ namespace AudioWidgets
 				];
 		}
 
-		TSharedRef<ITableRow> MakeMeterValueTileWidget(TSharedRef<const FLoudnessMeterWidgetView::FLoudnessMetric> InItem, const TSharedRef<STableViewBase>& OwnerTable, FLoudnessMeterWidgetView::FLoudnessScaleParams LoudnessScaleParams)
+		TSharedRef<ITableRow> MakeMeterValueTileWidget(FLoudnessMetricRef InItem, const TSharedRef<STableViewBase>& OwnerTable, FLoudnessMeterWidgetView::FLoudnessScaleParams LoudnessScaleParams)
 		{
-			using FLoudnessMetricRef = TSharedRef<const FLoudnessMeterWidgetView::FLoudnessMetric>;
-
 			static const FSlateRoundedBoxBrush RoundedPanelBrush(FStyleColors::Panel, 4.0f);
 
 			const FLoudnessValueConverter LoudnessValueConverter(LoudnessScaleParams);
@@ -267,16 +267,16 @@ namespace AudioWidgets
 
 	FLoudnessMeterWidgetView::FLoudnessMeterWidgetView()
 		: LoudnessMetrics(MakeShared<FLoudnessMetricRefArray>())
-		, NumericValues(MakeShared<FLoudnessMetricRefArray>())
-		, MeterValues(MakeShared<FLoudnessMetricRefArray>())
+		, ValuesCollectionView(MakeShared<FLoudnessMetricCollectionView>(LoudnessMetrics, [](const FLoudnessMetricRef& LoudnessMetricRef) { return LoudnessMetricRef->bShowValue.Get(); }))
+		, MetersCollectionView(MakeShared<FLoudnessMetricCollectionView>(LoudnessMetrics, [](const FLoudnessMetricRef& LoudnessMetricRef) { return LoudnessMetricRef->bShowMeter.Get(); }))
 	{
 	}
 
 	FLoudnessMeterWidgetView::~FLoudnessMeterWidgetView()
 	{
 		LoudnessMetrics->Reset();
-		NumericValues->Reset();
-		MeterValues->Reset();
+		ValuesCollectionView->ResetVisibleItems();
+		MetersCollectionView->ResetVisibleItems();
 	}
 
 	void FLoudnessMeterWidgetView::InitTimerPanel(const FTimerPanelParams& InTimerPanelParams)
@@ -289,6 +289,16 @@ namespace AudioWidgets
 		LoudnessScaleParams = InLoudnessScaleParams;
 	}
 
+	void FLoudnessMeterWidgetView::InitValuesDisplayOrder(const FDisplayOrderParams& InDisplayOrderParams)
+	{
+		ValuesCollectionView->SetDisplayOrderParams(InDisplayOrderParams);
+	}
+
+	void FLoudnessMeterWidgetView::InitMetersDisplayOrder(const FDisplayOrderParams& InDisplayOrderParams)
+	{
+		MetersCollectionView->SetDisplayOrderParams(InDisplayOrderParams);
+	}
+
 	TSharedRef<SWidget> FLoudnessMeterWidgetView::MakeWidget() const
 	{
 		using SLoudnessMetricTileView = STileView<FLoudnessMetricRef>;
@@ -470,10 +480,14 @@ namespace AudioWidgets
 			FText(),
 			FSlateIcon(FAppStyle::GetAppStyleSetName(), "LevelEditor.GameSettings"));
 
+		// Ensure collection views have been synchronized at least once before displaying UI:
+		ValuesCollectionView->RefreshVisibleItems();
+		MetersCollectionView->RefreshVisibleItems();
+
 		const auto GetTimerPanelVisibility = [bIsVisible = TimerPanelParams.bIsVisible]() { return (bIsVisible.Get()) ? EVisibility::Visible : EVisibility::Collapsed; };
-		const TAttribute<EVisibility> GetNumericValuesPanelVisibility = TAttribute<EVisibility>::CreateSPLambda(NumericValues, [&InNumericValues = NumericValues.Get()]()
+		const TAttribute<EVisibility> GetNumericValuesPanelVisibility = TAttribute<EVisibility>::CreateSPLambda(ValuesCollectionView, [&InValuesCollectionView = ValuesCollectionView.Get()]()
 			{
-				return (!InNumericValues.IsEmpty()) ? EVisibility::Visible : EVisibility::Collapsed;
+				return (InValuesCollectionView.HasVisibleItems()) ? EVisibility::Visible : EVisibility::Collapsed;
 			});
 
 		constexpr float MeterWidth = 79.0f;
@@ -481,7 +495,7 @@ namespace AudioWidgets
 
 		TSharedRef<SLoudnessMetricTileView> MeterTileView = SNew(SLoudnessMetricTileView)
 			.OnGenerateTile_Static(&LoudnessMeterWidgetView_Private::MakeMeterValueTileWidget, LoudnessScaleParams)
-			.ListItemsSource(MeterValues)
+			.ListItemsSource(MetersCollectionView->GetVisibleItems())
 			.ItemWidth(MeterWidth)
 			.ItemAlignment(EListItemAlignment::CenterAligned)
 			.SelectionMode(ESelectionMode::None)
@@ -569,7 +583,7 @@ namespace AudioWidgets
 						SNew(SLoudnessMetricTileView)
 							.Visibility(GetNumericValuesPanelVisibility)
 							.OnGenerateTile_Static(&LoudnessMeterWidgetView_Private::MakeNumericValueTileWidget, LoudnessScaleParams)
-							.ListItemsSource(NumericValues)
+							.ListItemsSource(ValuesCollectionView->GetVisibleItems())
 							.ItemHeight(60.0f)
 							.ItemWidth(60.0f)
 							.ItemAlignment(EListItemAlignment::CenterAligned)
@@ -598,51 +612,92 @@ namespace AudioWidgets
 	{
 		FLoudnessMetricRef LoudnessMetricRef = MakeShared<FLoudnessMetric>(LoudnessMetric);
 		LoudnessMetrics->Add(LoudnessMetricRef);
+	}
 
-		if (LoudnessMetric.bShowValue.Get())
-		{
-			NumericValues->Add(LoudnessMetricRef);
-		}
+	void FLoudnessMeterWidgetView::RefreshVisibleLoudnessMetrics()
+	{
+		ValuesCollectionView->RefreshVisibleItems();
+		MetersCollectionView->RefreshVisibleItems();
+	}
 
-		if (LoudnessMetric.bShowMeter.Get())
-		{
-			MeterValues->Add(LoudnessMetricRef);
-		}
+	FLoudnessMeterWidgetView::FLoudnessMetricCollectionView::FLoudnessMetricCollectionView(TSharedRef<FLoudnessMetricRefArray> InLoudnessMetrics, TFunction<bool(const FLoudnessMetricRef&)> InVisibilityPredicate)
+		: LoudnessMetrics(InLoudnessMetrics)
+		, VisibilityPredicate(InVisibilityPredicate)
+		, VisibleItems(MakeShared<FLoudnessMetricRefArray>())
+	{
 	}
 
-	void FLoudnessMeterWidgetView::RefreshVisibleLoudnessMetrics()
+	void FLoudnessMeterWidgetView::FLoudnessMetricCollectionView::RefreshVisibleItems()
 	{
-		for (const FLoudnessMetricRef& LoudnessMetricRef : LoudnessMetrics.Get())
-		{
-			const bool bShouldShowValue = LoudnessMetricRef->bShowValue.Get();
-			const bool bIsShowingValue = NumericValues->Contains(LoudnessMetricRef);
-			if (bIsShowingValue != bShouldShowValue)
+		const auto GetDesiredVisibleItems = [this]() -> TArray<FLoudnessMetricRef>
 			{
-				if (bShouldShowValue)
-				{
-					NumericValues->Add(LoudnessMetricRef);
-				}
-				else
+				const int64 PermutationIndex = DisplayOrderParams.PermutationIndex.Get();
+				const int32 NumLoudnessMetrics = LoudnessMetrics->Num();
+				const TArray<uint8> PermutationArray = GetPermutationArray(PermutationIndex, NumLoudnessMetrics);
+
+				TArray<FLoudnessMetricRef> DesiredVisibleItems;
+				DesiredVisibleItems.Reserve(NumLoudnessMetrics);
+				for (uint8 LoudnessMetricIndex : PermutationArray)
 				{
-					NumericValues->RemoveSingle(LoudnessMetricRef);
+					const FLoudnessMetricRef& LoudnessMetricRef = (*LoudnessMetrics)[LoudnessMetricIndex];
+					if (VisibilityPredicate(LoudnessMetricRef))
+					{
+						DesiredVisibleItems.Add(LoudnessMetricRef);
+					}
 				}
-			}
+				return DesiredVisibleItems;
+			};
+
+		const TArray<FLoudnessMetricRef> DesiredVisibleItems = GetDesiredVisibleItems();
+
+		// Sync VisibleItems to match DesiredVisibleItems:
+
+		const int32 NumDesired = DesiredVisibleItems.Num();
+		if (VisibleItems->Num() > NumDesired)
+		{
+			VisibleItems->RemoveAt(NumDesired, VisibleItems->Num() - NumDesired);
+		}
 
-			const bool bShouldShowMeter = LoudnessMetricRef->bShowMeter.Get();
-			const bool bIsShowingMeter = MeterValues->Contains(LoudnessMetricRef);
-			if (bIsShowingMeter != bShouldShowMeter)
+		for (int32 Index = 0; Index < NumDesired; Index++)
+		{
+			const FLoudnessMetricRef& LoudnessMetricRef = DesiredVisibleItems[Index];
+
+			if (Index >= VisibleItems->Num())
 			{
-				if (bShouldShowMeter)
-				{
-					MeterValues->Add(LoudnessMetricRef);
-				}
-				else
-				{
-					MeterValues->RemoveSingle(LoudnessMetricRef);
-				}
+				VisibleItems->Add(LoudnessMetricRef);
+			}
+			else if ((*VisibleItems)[Index] != LoudnessMetricRef)
+			{
+				VisibleItems->RemoveAt(Index);
+				VisibleItems->EmplaceAt(Index, LoudnessMetricRef);
 			}
 		}
 	}
+
+	uint64 FLoudnessMeterWidgetView::FLoudnessMetricCollectionView::GetPermutationIndex(TArrayView<uint8> PermutationArray)
+	{
+		using namespace Permutations;
+
+		const uint64 NumPermuations = Factorial(PermutationArray.Num());
+		const uint64 MaxPermutationIndex = NumPermuations - 1;
+
+		// We alter the returned index so that the permutation (0,1,2,...) is represented as PermutationIndex zero:
+		return MaxPermutationIndex - AnalyzePermutation(PermutationArray);
+	}
+
+	TArray<uint8> FLoudnessMeterWidgetView::FLoudnessMetricCollectionView::GetPermutationArray(uint64 PermutationIndex, uint8 NumValues)
+	{
+		using namespace Permutations;
+
+		const uint64 NumPermuations = Factorial(NumValues);
+		const uint64 MaxPermutationIndex = NumPermuations - 1;
+
+		// Ensure valid bounds:
+		PermutationIndex = PermutationIndex % NumPermuations;
+
+		// We use a reversed index so that a given PermutationIndex of zero will result in the permutation (0,1,2,...):
+		return GeneratePermutation(MaxPermutationIndex - PermutationIndex, NumValues);
+	}
 } // namespace AudioWidgets
 
 #undef LOCTEXT_NAMESPACE
diff --git a/Engine/Plugins/Runtime/AudioWidgets/Source/AudioWidgets/Private/LoudnessMeterWidgetView.h b/Engine/Plugins/Runtime/AudioWidgets/Source/AudioWidgets/Private/LoudnessMeterWidgetView.h
--- a/Engine/Plugins/Runtime/AudioWidgets/Source/AudioWidgets/Private/LoudnessMeterWidgetView.h
+++ b/Engine/Plugins/Runtime/AudioWidgets/Source/AudioWidgets/Private/LoudnessMeterWidgetView.h
@@ -62,11 +62,21 @@ namespace AudioWidgets
 			FSimpleDelegate OnRelativeScaleToggleRequested;
 		};
 
+		DECLARE_DELEGATE_OneParam(FOnPermutationChanged, uint64 /*NewPermutationIndex*/);
+
+		struct FDisplayOrderParams
+		{
+			TAttribute<uint64> PermutationIndex = 0;
+			FOnPermutationChanged OnPermutationChanged;
+		};
+
 		FLoudnessMeterWidgetView();
 		~FLoudnessMeterWidgetView();
 
 		void InitTimerPanel(const FTimerPanelParams& InTimerPanelParams);
 		void InitLoudnessScale(const FLoudnessScaleParams& InLoudnessScaleParams);
+		void InitValuesDisplayOrder(const FDisplayOrderParams& InDisplayOrderParams);
+		void InitMetersDisplayOrder(const FDisplayOrderParams& InDisplayOrderParams);
 
 		TSharedRef<SWidget> MakeWidget() const;
 
@@ -77,9 +87,30 @@ namespace AudioWidgets
 		using FLoudnessMetricRef = TSharedRef<const FLoudnessMetric>;
 		using FLoudnessMetricRefArray = UE::Slate::Containers::TObservableArray<FLoudnessMetricRef>;
 
+		class FLoudnessMetricCollectionView : public TSharedFromThis<FLoudnessMetricCollectionView>
+		{
+		public:
+			FLoudnessMetricCollectionView(TSharedRef<FLoudnessMetricRefArray> InLoudnessMetrics, TFunction<bool(const FLoudnessMetricRef&)> InVisibilityPredicate);
+
+			void SetDisplayOrderParams(const FDisplayOrderParams& InDisplayOrderParams) { DisplayOrderParams = InDisplayOrderParams; }
+			void RefreshVisibleItems();
+			void ResetVisibleItems() { VisibleItems->Reset(); }
+			TSharedRef<FLoudnessMetricRefArray> GetVisibleItems() { return VisibleItems; }
+			bool HasVisibleItems() const { return !VisibleItems->IsEmpty(); }
+
+		private:
+			static uint64 GetPermutationIndex(TArrayView<uint8> PermutationArray);
+			static TArray<uint8> GetPermutationArray(uint64 PermutationIndex, uint8 NumValues);
+
+			const TSharedRef<FLoudnessMetricRefArray> LoudnessMetrics;
+			const TFunction<bool(const FLoudnessMetricRef&)> VisibilityPredicate;
+			const TSharedRef<FLoudnessMetricRefArray> VisibleItems;
+			FDisplayOrderParams DisplayOrderParams;
+		};
+
 		TSharedRef<FLoudnessMetricRefArray> LoudnessMetrics;
-		TSharedRef<FLoudnessMetricRefArray> NumericValues;
-		TSharedRef<FLoudnessMetricRefArray> MeterValues;
+		TSharedRef<FLoudnessMetricCollectionView> ValuesCollectionView;
+		TSharedRef<FLoudnessMetricCollectionView> MetersCollectionView;
 		FTimerPanelParams TimerPanelParams;
 		FLoudnessScaleParams LoudnessScaleParams;
 	};
diff --git a/Engine/Plugins/Runtime/AudioWidgets/Source/AudioWidgets/Private/Permutations.cpp b/Engine/Plugins/Runtime/AudioWidgets/Source/AudioWidgets/Private/Permutations.cpp
new file mode 100644
--- /dev/null
+++ b/Engine/Plugins/Runtime/AudioWidgets/Source/AudioWidgets/Private/Permutations.cpp
@@ -0,0 +1,51 @@
+// Copyright Epic Games, Inc. All Rights Reserved.
+
+#include "Permutations.h"
+
+#include "Algo/MaxElement.h"
+
+namespace AudioWidgets::Permutations
+{
+	uint64 AnalyzePermutation(TArrayView<uint8> PermutationArray)
+	{
+		check(PermutationArray.Num() >= 0);
+		check(PermutationArray.Num() <= MaxFactorial);
+
+		// See Knuth, Vol. 2, 3rd ed., Section 3.3.2, pp. 65-66, Algorithm P (Analyze a permutation).
+
+		uint64 F = 0;
+		for (uint8 R = PermutationArray.Num(); R > 1; R--)
+		{
+			PermutationArray.LeftInline(R);
+			const uint64 S = Algo::MaxElement(PermutationArray) - PermutationArray.GetData();
+			F = R * F + S;
+			Swap(PermutationArray[R - 1], PermutationArray[S]);
+		}
+		return F;
+	}
+
+	TArray<uint8> GeneratePermutation(uint64 PermutationID, uint8 NumValues)
+	{
+		check(NumValues <= MaxFactorial);
+		check(PermutationID < Factorial(NumValues));
+
+		// Algorithm P, reverse
+
+		TArray<uint8> PermutationArray;
+		PermutationArray.Reserve(NumValues);
+		while (PermutationArray.Num() < NumValues)
+		{
+			PermutationArray.Add(PermutationArray.Num());
+		}
+
+		uint64 F = PermutationID;
+		for (uint8 R = 2; R <= NumValues; R++)
+		{
+			const uint64 S = F % R;
+			F /= R;
+			PermutationArray.Swap(R - 1, S);
+		}
+		return PermutationArray;
+	}
+
+} // namespace AudioWidgets::Permutations
diff --git a/Engine/Plugins/Runtime/AudioWidgets/Source/AudioWidgets/Private/Permutations.h b/Engine/Plugins/Runtime/AudioWidgets/Source/AudioWidgets/Private/Permutations.h
new file mode 100644
--- /dev/null
+++ b/Engine/Plugins/Runtime/AudioWidgets/Source/AudioWidgets/Private/Permutations.h
@@ -0,0 +1,46 @@
+// Copyright Epic Games, Inc. All Rights Reserved.
+
+#pragma once
+
+#include "Containers/ArrayView.h"
+
+namespace AudioWidgets::Permutations
+{
+	constexpr uint8 MaxFactorial = 20;
+
+	/**
+	 * Compute exact factorial. Only valid for small non-negative integers.
+	 */
+	constexpr uint64 Factorial(uint8 N)
+	{
+		if (!ensure(N <= MaxFactorial))
+		{
+			return UINT64_MAX;
+		}
+
+		uint64 F = 1;
+		for (uint8 R = N; R > 1; R--)
+		{
+			F *= R;
+		}
+		return F;
+	}
+
+	/**
+	 * Get an integer that uniquely represents the given permutation.
+	 * 
+	 * @param PermutationArray The elements of the array should be distinct values in the interval [0, N-1] where N == PermutationArray.Num(). On return, the array will have been sorted.
+	 * @return A value in the interval [0, N!-1] where N == PermutationArray.Num(). If N > MaxFactorial, the result is undefined.
+	 */
+	uint64 AnalyzePermutation(TArrayView<uint8> PermutationArray);
+
+	/**
+	 * Create the permutation that is uniquely represented by the given integer.
+	 * 
+	 * @param PermutationID A value in the interval [0, N!-1] where N == NumValues.
+	 * @param NumValues The desired number of values in the permutation, must be no greater than MaxFactorial.
+	 * @return An array of distinct values with length == NumValues. If NumValues > MaxFactorial, the result is undefined.
+	 */
+	TArray<uint8> GeneratePermutation(uint64 PermutationID, uint8 NumValues);
+
+} // namespace AudioWidgets::Permutations
diff --git a/Engine/Plugins/Runtime/AudioWidgets/Source/AudioWidgets/Public/LoudnessMeterRackUnitSettings.h b/Engine/Plugins/Runtime/AudioWidgets/Source/AudioWidgets/Public/LoudnessMeterRackUnitSettings.h
--- a/Engine/Plugins/Runtime/AudioWidgets/Source/AudioWidgets/Public/LoudnessMeterRackUnitSettings.h
+++ b/Engine/Plugins/Runtime/AudioWidgets/Source/AudioWidgets/Public/LoudnessMeterRackUnitSettings.h
@@ -47,4 +47,10 @@ struct FLoudnessMeterRackUnitSettings
 
 	UPROPERTY(EditAnywhere, config, Category = LoudnessMeter)
 	FLoudnessMetricDisplayOptions MomentaryLoudness = { .bShowValue = false, .bShowMeter = true };
+
+	UPROPERTY(EditAnywhere, config, Category = LoudnessMeter, meta = (ClampMin = "0", ClampMax = "5", UIMin = "0", UIMax = "5"))
+	uint64 ValuesOrderingPermutation = 0; // Max == 3! - 1
+
+	UPROPERTY(EditAnywhere, config, Category = LoudnessMeter, meta = (ClampMin = "0", ClampMax = "5", UIMin = "0", UIMax = "5"))
+	uint64 MetersOrderingPermutation = 0; // Max == 3! - 1
 };
