[Audio Widgets] Loudness Meter - Slate code to support the presentation aspects of dragging loudness metric tiles. Actual drag/drop operation behavior to be added in forthcoming CL.
- Added FOnMakeTileInnerContent delegate to allow for both the actual STableRow and the FDragDropOperation::GetDefaultDecorator() to draw the loudness metric value/meter tile.

#Synced-CL 48500656

diff --git a/Engine/Plugins/Runtime/AudioWidgets/Source/AudioWidgets/Private/LoudnessMeterWidgetView.cpp b/Engine/Plugins/Runtime/AudioWidgets/Source/AudioWidgets/Private/LoudnessMeterWidgetView.cpp
--- a/Engine/Plugins/Runtime/AudioWidgets/Source/AudioWidgets/Private/LoudnessMeterWidgetView.cpp
+++ b/Engine/Plugins/Runtime/AudioWidgets/Source/AudioWidgets/Private/LoudnessMeterWidgetView.cpp
@@ -142,6 +142,9 @@ namespace AudioWidgets
 			TAttribute<int> Target;
 		};
 
+		/**
+		 * Creates a widget that displays the numeric value of the given loudness metric, with units.
+		 */
 		TSharedRef<SWidget> MakeValueReadoutWidget(FLoudnessMetricRef InItem, const FLoudnessScaleParams& LoudnessScaleParams)
 		{
 			const FLoudnessValueConverter LoudnessValueConverter(LoudnessScaleParams);
@@ -173,45 +176,260 @@ namespace AudioWidgets
 				];
 		}
 
-		TSharedRef<ITableRow> MakeNumericValueTileWidget(FLoudnessMetricRef InItem, const TSharedRef<STableViewBase>& OwnerTable, FLoudnessScaleParams LoudnessScaleParams)
+		/**
+		 * Creates a widget with the actual content of a Value tile, containing the numeric value readout widget and the loudness metric display name.
+		 */
+		TSharedRef<SWidget> MakeValueTileInnerContent(FLoudnessMetricRef InItem, FLoudnessScaleParams LoudnessScaleParams)
 		{
-			return SNew(STableRow<FLoudnessMetricRef>, OwnerTable)
+			return SNew(SVerticalBox)
+				+ SVerticalBox::Slot()
+				.MinHeight(8)
+				.MaxHeight(8)
 				[
-					SNew(SVerticalBox)
-						+ SVerticalBox::Slot()
-						.MinHeight(8)
-						.MaxHeight(8)
-						[
-							SNew(SSpacer)
-						]
-						+ SVerticalBox::Slot()
-						.AutoHeight()
-						.HAlign(EHorizontalAlignment::HAlign_Center)
-						[
-							MakeValueReadoutWidget(InItem, LoudnessScaleParams)
-						]
-						+ SVerticalBox::Slot()
-						.MinHeight(24)
-						.MaxHeight(24)
-						.HAlign(EHorizontalAlignment::HAlign_Center)
-						.VAlign(EVerticalAlignment::VAlign_Center)
+					SNew(SSpacer)
+						.Size(FVector2D(60.0f, 8.0f))
+				]
+				+ SVerticalBox::Slot()
+				.AutoHeight()
+				.HAlign(EHorizontalAlignment::HAlign_Center)
+				[
+					MakeValueReadoutWidget(InItem, LoudnessScaleParams)
+				]
+				+ SVerticalBox::Slot()
+				.MinHeight(24)
+				.MaxHeight(24)
+				.HAlign(EHorizontalAlignment::HAlign_Center)
+				.VAlign(EVerticalAlignment::VAlign_Center)
+				[
+					SNew(STextBlock)
+						.Text(InItem->DisplayName)
+						.Font(FStyleDefaults::GetFontInfo(8))
+						.AutoWrapText(true)
+						.Justification(ETextJustify::Center)
+				];
+		}
+
+		/**
+		 * Creates a widget with the actual content of a Meter tile, containing the numeric value readout widget, the actual meter widget, and the loudness metric display name.
+		 */
+		TSharedRef<SWidget> MakeMeterTileInnerContent(FLoudnessMetricRef InItem, FLoudnessScaleParams LoudnessScaleParams)
+		{
+			const FLoudnessValueConverter LoudnessValueConverter(LoudnessScaleParams);
+
+			return SNew(SVerticalBox)
+				+ SVerticalBox::Slot()
+				.AutoHeight()
+				.HAlign(EHorizontalAlignment::HAlign_Center)
+				[
+					MakeValueReadoutWidget(InItem, LoudnessScaleParams)
+				]
+				+ SVerticalBox::Slot()
+				.HAlign(EHorizontalAlignment::HAlign_Center)
+				[
+					SNew(SVariableRangeMeterWidget)
+						.MeterValueColor_Lambda([InItem]() { return InItem->Color.Get(); })
+						.MeterChannelInfo_Lambda([InItem, LoudnessValueConverter]() -> TArray<FAudioMeterChannelInfo>
+							{
+								constexpr float DefaultInvalidValue = MIN_VOLUME_DECIBELS;
+								const TOptional<float> ConvertedValue = LoudnessValueConverter.ConvertValue(InItem->Value.Get());
+								return { { ConvertedValue.Get(DefaultInvalidValue), DefaultInvalidValue, DefaultInvalidValue } };
+							})
+						.Range(LoudnessScaleParams.Range)
+						.Offset_Lambda([Offset = LoudnessScaleParams.Offset, LoudnessValueConverter] { return LoudnessValueConverter.ConvertValue(Offset.Get()); })
+				]
+				+ SVerticalBox::Slot()
+				.AutoHeight()
+				.MinHeight(28)
+				.HAlign(EHorizontalAlignment::HAlign_Center)
+				.VAlign(EVerticalAlignment::VAlign_Center)
+				[
+					SNew(STextBlock)
+						.Text_Lambda([InItem]() { return InItem->DisplayName; })
+						.Font(FStyleDefaults::GetFontInfo(8))
+						.AutoWrapText(true)
+						.Justification(ETextJustify::Center)
+				];
+		}
+
+		DECLARE_DELEGATE_RetVal_OneParam(TSharedRef<SWidget>, FOnMakeTileInnerContent, FLoudnessMetricRef);
+
+		class SLoudnessMetricTileView;
+
+		/**
+		 * Represents a loudness metric that is being dragged from a given source tile view.
+		 * A decorator renders the content of the tile, with a rounded border selection highlight.
+		 */
+		class FLoudnessValueDragDropOp : public FDragDropOperation
+		{
+		public:
+			DRAG_DROP_OPERATOR_TYPE(FLoudnessValueDragDropOp, FDragDropOperation)
+
+			FLoudnessValueDragDropOp(TSharedRef<SLoudnessMetricTileView> InSourceTileView, FLoudnessMetricRef InLoudnessMetric, FOnMakeTileInnerContent InOnMakeTileInnerContent)
+				: WeakSourceTileView(InSourceTileView.ToWeakPtr())
+				, LoudnessMetric(InLoudnessMetric)
+				, OnMakeTileInnerContent(InOnMakeTileInnerContent)
+			{
+				Construct();
+			}
+
+			virtual FCursorReply OnCursorQuery() override;
+			virtual TSharedPtr<SWidget> GetDefaultDecorator() const override;
+
+			TSharedPtr<SLoudnessMetricTileView> GetSourceTileView() const { return WeakSourceTileView.Pin(); }
+			FLoudnessMetricRef GetLoudnessMetric() const { return LoudnessMetric; }
+
+		private:
+			const TWeakPtr<SLoudnessMetricTileView> WeakSourceTileView;
+			const FLoudnessMetricRef LoudnessMetric;
+			const FOnMakeTileInnerContent OnMakeTileInnerContent;
+		};
+
+		FCursorReply FLoudnessValueDragDropOp::OnCursorQuery()
+		{
+			return FCursorReply::Cursor(EMouseCursor::GrabHand);
+		}
+
+		TSharedPtr<SWidget> FLoudnessValueDragDropOp::GetDefaultDecorator() const
+		{
+			static const FSlateColorBrush BackgroundPanelBrush(FStyleColors::Panel);
+			static const FSlateRoundedBoxBrush BorderBrush(FStyleColors::Select, 4.0f);
+			static const FSlateRoundedBoxBrush InteriorPanelBrush(FStyleColors::Panel, 4.0f);
+
+			TSharedRef<SWidget> TileInnerContent = (OnMakeTileInnerContent.IsBound()) ? OnMakeTileInnerContent.Execute(LoudnessMetric) : SNullWidget::NullWidget;
+
+			return SNew(SBox)
+				.MaxDesiredHeight(180.0f)
+				.MinDesiredWidth(66.0f)
+				[
+					SNew(SBorder)
+						.BorderImage(&BackgroundPanelBrush)
+						.Padding(0.0f)
 						[
-							SNew(STextBlock)
-								.Text_Lambda([InItem]() { return InItem->DisplayName; })
-								.Font(FStyleDefaults::GetFontInfo(8))
-								.AutoWrapText(true)
-								.Justification(ETextJustify::Center)
+							SNew(SBorder)
+								.BorderImage(&BorderBrush)
+								.Padding(1.0f)
+								[
+									SNew(SBorder)
+										.BorderImage(&InteriorPanelBrush)
+										[
+											TileInnerContent
+										]
+								]
 						]
 				];
 		}
 
-		TSharedRef<ITableRow> MakeMeterValueTileWidget(FLoudnessMetricRef InItem, const TSharedRef<STableViewBase>& OwnerTable, FLoudnessMeterWidgetView::FLoudnessScaleParams LoudnessScaleParams)
+
+		/**
+		 * TileView widget for displaying an observable collection of loudness metrics.
+		 */
+		class SLoudnessMetricTileView : public STileView<FLoudnessMetricRef>
+		{
+		public:
+			void Construct(const FArguments& InArgs, FOnMakeTileInnerContent InOnMakeTileInnerContent)
+			{
+				STileView<FLoudnessMetricRef>::Construct(InArgs);
+
+				OnMakeTileInnerContent = InOnMakeTileInnerContent;
+			}
+
+			TSharedRef<SWidget> MakeTileInnerContent(const FLoudnessMetricRef& LoudnessMetric) const
+			{
+				return (OnMakeTileInnerContent.IsBound()) ? OnMakeTileInnerContent.Execute(LoudnessMetric) : SNullWidget::NullWidget;
+			}
+
+		private:
+			FOnMakeTileInnerContent OnMakeTileInnerContent;
+		};
+
+
+		/**
+		 * Tile widget for a single loudness metric.
+		 */
+		class SLoudnessMetricTile : public STableRow<FLoudnessMetricRef>
+		{
+		public:
+			FLinearColor GetOverlayColor() const
+			{
+				if (IsBeingDragged())
+				{
+					return FStyleColors::Recessed.GetSpecifiedColor().CopyWithNewOpacity(0.9f);
+				}
+				return FLinearColor::Transparent;
+			}
+
+			bool IsBeingDragged() const { return WeakDragDropOperation.IsValid(); }
+
+		private:
+			TWeakPtr<FDragDropOperation> WeakDragDropOperation; // Will be set when a drag is detected.
+		};
+
+
+		/**
+		 * Helper function to create an empty tile widget, with common styling applied.
+		 */
+		TSharedRef<SLoudnessMetricTile> MakeEmptyTileWidget(const TSharedRef<STableViewBase>& OwnerTable)
+		{
+			const auto MakeTableRowStyle = []() -> FTableRowStyle
+				{
+					const FSlateRoundedBoxBrush RoundedPanelBrush(FStyleColors::Panel, 4.0f);
+
+					FTableRowStyle TableRowStyle = FAppStyle::Get().GetWidgetStyle<FTableRowStyle>("TableView.Row");
+					TableRowStyle.SetEvenRowBackgroundHoveredBrush(RoundedPanelBrush);
+					TableRowStyle.SetOddRowBackgroundHoveredBrush(RoundedPanelBrush);
+					TableRowStyle.SetSelectorFocusedBrush(FSlateNoResource());
+					return TableRowStyle;
+				};
+			static FTableRowStyle LoudnessMetricTableRowStyle = MakeTableRowStyle();
+
+			return SNew(SLoudnessMetricTile, OwnerTable)
+				.ShowSelection(false)
+				.Style(&LoudnessMetricTableRowStyle);
+		}
+
+		/**
+		 * Callback to create a tile widget for the numeric Values tile view. A color block overlay will be used to darken the tile content when it's being dragged.
+		 */
+		TSharedRef<ITableRow> MakeNumericValueTileWidget(FLoudnessMetricRef InItem, const TSharedRef<STableViewBase>& OwnerTable)
+		{
+			TSharedRef<SLoudnessMetricTileView> OwnerTileView = StaticCastSharedRef<SLoudnessMetricTileView>(OwnerTable);
+
+			TSharedRef<SLoudnessMetricTile> TileWidget = MakeEmptyTileWidget(OwnerTileView);
+
+			TileWidget->SetContent(
+				SNew(SOverlay)
+				+ SOverlay::Slot()
+				[
+					OwnerTileView->MakeTileInnerContent(InItem)
+				]
+				+ SOverlay::Slot()
+				[
+					SNew(SColorBlock)
+						.Color(TileWidget, &SLoudnessMetricTile::GetOverlayColor)
+				]);
+
+			return TileWidget;
+		}
+
+		/**
+		 * Callback to create a tile widget for the Meters tile view. A color block overlay will be used to darken the tile content when it's being dragged.
+		 */
+		TSharedRef<ITableRow> MakeMeterValueTileWidget(FLoudnessMetricRef InItem, const TSharedRef<STableViewBase>& OwnerTable)
 		{
 			static const FSlateRoundedBoxBrush RoundedPanelBrush(FStyleColors::Panel, 4.0f);
 
-			const FLoudnessValueConverter LoudnessValueConverter(LoudnessScaleParams);
+			TSharedRef<SLoudnessMetricTileView> OwnerTileView = StaticCastSharedRef<SLoudnessMetricTileView>(OwnerTable);
+
+			TSharedRef<SLoudnessMetricTile> TileWidget = MakeEmptyTileWidget(OwnerTileView);
 
-			return SNew(STableRow<FLoudnessMetricRef>, OwnerTable)
+			TileWidget->SetContent(
+				SNew(SOverlay)
+				+ SOverlay::Slot()
+				[
+					SNew(SColorBlock)
+						.Color(FStyleColors::Recessed.GetSpecifiedColor())
+				]
+				+ SOverlay::Slot()
 				[
 					SNew(SBox)
 						.VAlign(EVerticalAlignment::VAlign_Center)
@@ -221,42 +439,17 @@ namespace AudioWidgets
 								.BorderImage(&RoundedPanelBrush)
 								.Padding(0.0f, 8.0f)
 								[
-									SNew(SVerticalBox)
-										+ SVerticalBox::Slot()
-										.AutoHeight()
-										.HAlign(EHorizontalAlignment::HAlign_Center)
-										[
-											MakeValueReadoutWidget(InItem, LoudnessScaleParams)
-										]
-										+ SVerticalBox::Slot()
-										.HAlign(EHorizontalAlignment::HAlign_Center)
-										[
-											SNew(SVariableRangeMeterWidget)
-												.MeterValueColor_Lambda([InItem]() { return InItem->Color.Get(); })
-												.MeterChannelInfo_Lambda([InItem, LoudnessValueConverter]() -> TArray<FAudioMeterChannelInfo>
-													{
-														constexpr float DefaultInvalidValue = MIN_VOLUME_DECIBELS;
-														const TOptional<float> ConvertedValue = LoudnessValueConverter.ConvertValue(InItem->Value.Get());
-														return { { ConvertedValue.Get(DefaultInvalidValue), DefaultInvalidValue, DefaultInvalidValue } };
-													})
-												.Range(LoudnessScaleParams.Range)
-												.Offset_Lambda([Offset = LoudnessScaleParams.Offset, LoudnessValueConverter] { return LoudnessValueConverter.ConvertValue(Offset.Get()); })
-										]
-										+ SVerticalBox::Slot()
-										.AutoHeight()
-										.MinHeight(28)
-										.HAlign(EHorizontalAlignment::HAlign_Center)
-										.VAlign(EVerticalAlignment::VAlign_Center)
-										[
-											SNew(STextBlock)
-												.Text_Lambda([InItem]() { return InItem->DisplayName; })
-												.Font(FStyleDefaults::GetFontInfo(8))
-												.AutoWrapText(true)
-												.Justification(ETextJustify::Center)
-										]
+									OwnerTileView->MakeTileInnerContent(InItem)
 								]
 						]
-				];
+				]
+				+ SOverlay::Slot()
+				[
+					SNew(SColorBlock)
+						.Color(TileWidget, &SLoudnessMetricTile::GetOverlayColor)
+				]);
+
+			return TileWidget;
 		}
 	} // namespace LoudnessMeterWidgetView_Private
 
@@ -301,7 +494,7 @@ namespace AudioWidgets
 
 	TSharedRef<SWidget> FLoudnessMeterWidgetView::MakeWidget() const
 	{
-		using SLoudnessMetricTileView = STileView<FLoudnessMetricRef>;
+		using namespace LoudnessMeterWidgetView_Private;
 
 		static FSlateColorBrush BackgroundBrush(FStyleColors::Recessed);
 
@@ -493,8 +686,9 @@ namespace AudioWidgets
 		constexpr float MeterWidth = 79.0f;
 		constexpr float MeterMinHeight = 160.0f;
 
-		TSharedRef<SLoudnessMetricTileView> MeterTileView = SNew(SLoudnessMetricTileView)
-			.OnGenerateTile_Static(&LoudnessMeterWidgetView_Private::MakeMeterValueTileWidget, LoudnessScaleParams)
+		TSharedRef<SLoudnessMetricTileView> MeterTileView = SNew(SLoudnessMetricTileView,
+				FOnMakeTileInnerContent::CreateStatic(&MakeMeterTileInnerContent, LoudnessScaleParams))
+			.OnGenerateTile_Static(&MakeMeterValueTileWidget)
 			.ListItemsSource(MetersCollectionView->GetVisibleItems())
 			.ItemWidth(MeterWidth)
 			.ItemAlignment(EListItemAlignment::CenterAligned)
@@ -580,9 +774,10 @@ namespace AudioWidgets
 					+ SVerticalBox::Slot()
 					.AutoHeight()
 					[
-						SNew(SLoudnessMetricTileView)
+						SNew(SLoudnessMetricTileView,
+								FOnMakeTileInnerContent::CreateStatic(&MakeValueTileInnerContent, LoudnessScaleParams))
 							.Visibility(GetNumericValuesPanelVisibility)
-							.OnGenerateTile_Static(&LoudnessMeterWidgetView_Private::MakeNumericValueTileWidget, LoudnessScaleParams)
+							.OnGenerateTile_Static(&MakeNumericValueTileWidget)
 							.ListItemsSource(ValuesCollectionView->GetVisibleItems())
 							.ItemHeight(60.0f)
 							.ItemWidth(60.0f)
