[Audio Widgets] Loudness Meter - Interaction logic for drag drop reordering of items.

#Synced-CL 49282079

diff --git a/Engine/Plugins/Runtime/AudioWidgets/Source/AudioWidgets/Private/LoudnessMeterWidgetView.cpp b/Engine/Plugins/Runtime/AudioWidgets/Source/AudioWidgets/Private/LoudnessMeterWidgetView.cpp
--- a/Engine/Plugins/Runtime/AudioWidgets/Source/AudioWidgets/Private/LoudnessMeterWidgetView.cpp
+++ b/Engine/Plugins/Runtime/AudioWidgets/Source/AudioWidgets/Private/LoudnessMeterWidgetView.cpp
@@ -11,6 +11,7 @@
 #include "Framework/MultiBox/MultiBoxBuilder.h"
 #include "Permutations.h"
 #include "SAudioMeterWidget.h"
+#include "SDragReorderableTileView.h"
 #include "Styling/StyleColors.h"
 #include "Widgets/Colors/SColorBlock.h"
 #include "Widgets/Colors/SColorPicker.h"
@@ -18,7 +19,6 @@
 #include "Widgets/Input/SNumericEntryBox.h"
 #include "Widgets/Layout/SSeparator.h"
 #include "Widgets/Layout/SSpacer.h"
-#include "Widgets/Views/STileView.h"
 
 #include <limits>
 
@@ -253,42 +253,34 @@ namespace AudioWidgets
 
 		DECLARE_DELEGATE_RetVal_OneParam(TSharedRef<SWidget>, FOnMakeTileInnerContent, FLoudnessMetricRef);
 
-		class SLoudnessMetricTileView;
-
 		/**
 		 * Represents a loudness metric that is being dragged from a given source tile view.
 		 * A decorator renders the content of the tile, with a rounded border selection highlight.
 		 */
-		class FLoudnessValueDragDropOp : public FDragDropOperation
+		class FLoudnessValueDragDropOp : public TItemDragDropOperation<FLoudnessMetricRef, FLoudnessValueDragDropOp>
 		{
 		public:
-			DRAG_DROP_OPERATOR_TYPE(FLoudnessValueDragDropOp, FDragDropOperation)
+			using Base = TItemDragDropOperation<FLoudnessMetricRef, FLoudnessValueDragDropOp>;
+
+			DRAG_DROP_OPERATOR_TYPE(FLoudnessValueDragDropOp, Base)
 
-			FLoudnessValueDragDropOp(TSharedRef<SLoudnessMetricTileView> InSourceTileView, FLoudnessMetricRef InLoudnessMetric, FOnMakeTileInnerContent InOnMakeTileInnerContent)
-				: WeakSourceTileView(InSourceTileView.ToWeakPtr())
+			FLoudnessValueDragDropOp(TSharedRef<Base::STileViewType> InSourceTileView, FLoudnessMetricRef InLoudnessMetric, FOnMakeTileInnerContent InOnMakeTileInnerContent)
+				: Base(InSourceTileView)
 				, LoudnessMetric(InLoudnessMetric)
 				, OnMakeTileInnerContent(InOnMakeTileInnerContent)
 			{
 				Construct();
 			}
 
-			virtual FCursorReply OnCursorQuery() override;
 			virtual TSharedPtr<SWidget> GetDefaultDecorator() const override;
 
-			TSharedPtr<SLoudnessMetricTileView> GetSourceTileView() const { return WeakSourceTileView.Pin(); }
-			FLoudnessMetricRef GetLoudnessMetric() const { return LoudnessMetric; }
+			virtual const FLoudnessMetricRef& GetItem() const final { return LoudnessMetric; }
 
 		private:
-			const TWeakPtr<SLoudnessMetricTileView> WeakSourceTileView;
 			const FLoudnessMetricRef LoudnessMetric;
 			const FOnMakeTileInnerContent OnMakeTileInnerContent;
 		};
 
-		FCursorReply FLoudnessValueDragDropOp::OnCursorQuery()
-		{
-			return FCursorReply::Cursor(EMouseCursor::GrabHand);
-		}
-
 		TSharedPtr<SWidget> FLoudnessValueDragDropOp::GetDefaultDecorator() const
 		{
 			static const FSlateColorBrush BackgroundPanelBrush(FStyleColors::Panel);
@@ -319,16 +311,40 @@ namespace AudioWidgets
 				];
 		}
 
+		/**
+		 * Delegate adapter to execute the necessary delegate to show or hide the numeric value of the given loudness metric.
+		 */
+		void SetShowLoudnessMetricValue(const FLoudnessMetricRef& LoudnessMetric, bool bShowValue)
+		{
+			if (LoudnessMetric->bShowValue.Get() != bShowValue)
+			{
+				LoudnessMetric->OnShowValueToggleRequested.ExecuteIfBound();
+			}
+		}
+
+		/**
+		 * Delegate adapter to execute the necessary delegate to show or hide the meter for the given loudness metric.
+		 */
+		void SetShowLoudnessMetricMeter(const FLoudnessMetricRef& LoudnessMetric, bool bShowMeter)
+		{
+			if (LoudnessMetric->bShowMeter.Get() != bShowMeter)
+			{
+				LoudnessMetric->OnShowMeterToggleRequested.ExecuteIfBound();
+			}
+		}
 
 		/**
 		 * TileView widget for displaying an observable collection of loudness metrics.
+		 * Dropping of a dragged loudness metric onto the tile view is supported if the tile view is currently empty.
 		 */
-		class SLoudnessMetricTileView : public STileView<FLoudnessMetricRef>
+		class SLoudnessMetricTileView : public SDragReorderableTileView<FLoudnessValueDragDropOp>
 		{
 		public:
-			void Construct(const FArguments& InArgs, FOnMakeTileInnerContent InOnMakeTileInnerContent)
+			using Base = SDragReorderableTileView<FLoudnessValueDragDropOp>;
+
+			void Construct(const Base::FArguments& InArgs, FOnMakeTileInnerContent InOnMakeTileInnerContent, FOnSetShowItem InOnSetShowItem, FOnReorderItems InOnReorderItems)
 			{
-				STileView<FLoudnessMetricRef>::Construct(InArgs);
+				Base::Construct(InArgs, InOnSetShowItem, InOnReorderItems);
 
 				OnMakeTileInnerContent = InOnMakeTileInnerContent;
 			}
@@ -338,17 +354,45 @@ namespace AudioWidgets
 				return (OnMakeTileInnerContent.IsBound()) ? OnMakeTileInnerContent.Execute(LoudnessMetric) : SNullWidget::NullWidget;
 			}
 
+			virtual TSharedPtr<FLoudnessValueDragDropOp> CreateDragDropOperation() override
+			{
+				ensure(GetSelectionMode() == ESelectionMode::Single);
+				if (FLoudnessMetricRef* SelectedItem = SelectedItems.FindArbitraryElement()) // As we are using ESelectionMode::Single an arbitrary element is the only element.
+				{
+					TSharedRef<FLoudnessValueDragDropOp> DragDropOperation = MakeShared<FLoudnessValueDragDropOp>(SharedThis(this), *SelectedItem, OnMakeTileInnerContent);
+
+					// Hold a weak pointer to maintain awareness of whether a tile is currently being dragged:
+					WeakDragDropOperation = DragDropOperation;
+
+					return DragDropOperation;
+				}
+				return nullptr;
+			}
+
+			bool HasDragDropOperation() const { return WeakDragDropOperation.IsValid(); }
+
 		private:
 			FOnMakeTileInnerContent OnMakeTileInnerContent;
+			TWeakPtr<FDragDropOperation> WeakDragDropOperation;
 		};
 
 
 		/**
-		 * Tile widget for a single loudness metric.
+		 * Tile widget for a single loudness metric. Dragging and dropping tiles is supported.
 		 */
-		class SLoudnessMetricTile : public STableRow<FLoudnessMetricRef>
+		class SLoudnessMetricTile : public SDragReorderableTile<FLoudnessValueDragDropOp>
 		{
 		public:
+			const FSlateBrush* GetHoveredBorder() const
+			{
+				if (IsHovered())
+				{
+					static const FSlateRoundedBoxBrush RoundedHoverBrush(FStyleColors::Hover, 4.0f);
+					return &RoundedHoverBrush;
+				}
+				return nullptr;
+			}
+
 			FLinearColor GetOverlayColor() const
 			{
 				if (IsBeingDragged())
@@ -357,11 +401,6 @@ namespace AudioWidgets
 				}
 				return FLinearColor::Transparent;
 			}
-
-			bool IsBeingDragged() const { return WeakDragDropOperation.IsValid(); }
-
-		private:
-			TWeakPtr<FDragDropOperation> WeakDragDropOperation; // Will be set when a drag is detected.
 		};
 
 
@@ -433,13 +472,18 @@ namespace AudioWidgets
 				[
 					SNew(SBox)
 						.VAlign(EVerticalAlignment::VAlign_Center)
-						.Padding(4.0f)
+						.Padding(3.0f)
 						[
 							SNew(SBorder)
-								.BorderImage(&RoundedPanelBrush)
-								.Padding(0.0f, 8.0f)
+								.BorderImage(TileWidget, &SLoudnessMetricTile::GetHoveredBorder)
+								.Padding(1.0f)
 								[
-									OwnerTileView->MakeTileInnerContent(InItem)
+									SNew(SBorder)
+										.BorderImage(&RoundedPanelBrush)
+										.Padding(0.0f, 8.0f)
+										[
+											OwnerTileView->MakeTileInnerContent(InItem)
+										]
 								]
 						]
 				]
@@ -678,21 +722,19 @@ namespace AudioWidgets
 		MetersCollectionView->RefreshVisibleItems();
 
 		const auto GetTimerPanelVisibility = [bIsVisible = TimerPanelParams.bIsVisible]() { return (bIsVisible.Get()) ? EVisibility::Visible : EVisibility::Collapsed; };
-		const TAttribute<EVisibility> GetNumericValuesPanelVisibility = TAttribute<EVisibility>::CreateSPLambda(ValuesCollectionView, [&InValuesCollectionView = ValuesCollectionView.Get()]()
-			{
-				return (InValuesCollectionView.HasVisibleItems()) ? EVisibility::Visible : EVisibility::Collapsed;
-			});
 
 		constexpr float MeterWidth = 79.0f;
 		constexpr float MeterMinHeight = 160.0f;
 
 		TSharedRef<SLoudnessMetricTileView> MeterTileView = SNew(SLoudnessMetricTileView,
-				FOnMakeTileInnerContent::CreateStatic(&MakeMeterTileInnerContent, LoudnessScaleParams))
+				FOnMakeTileInnerContent::CreateStatic(&MakeMeterTileInnerContent, LoudnessScaleParams),
+				SLoudnessMetricTileView::FOnSetShowItem::CreateStatic(&SetShowLoudnessMetricMeter),
+				SLoudnessMetricTileView::FOnReorderItems::CreateSP(MetersCollectionView, &FLoudnessMetricCollectionView::DraggedItemReorder))
 			.OnGenerateTile_Static(&MakeMeterValueTileWidget)
 			.ListItemsSource(MetersCollectionView->GetVisibleItems())
 			.ItemWidth(MeterWidth)
 			.ItemAlignment(EListItemAlignment::CenterAligned)
-			.SelectionMode(ESelectionMode::None)
+			.SelectionMode(ESelectionMode::Single)
 			.ScrollbarVisibility(EVisibility::Collapsed);
 
 		MeterTileView->SetItemHeight(TAttribute<float>::CreateLambda([&TileView = *MeterTileView]()
@@ -705,6 +747,20 @@ namespace AudioWidgets
 				return FMath::Max(PanelSize.Y / HeightDivisor, MeterMinHeight);
 			}));
 
+		const TAttribute<EVisibility> GetNumericValuesPanelVisibility = TAttribute<EVisibility>::CreateSPLambda(ValuesCollectionView,
+			[&InValuesCollectionView = ValuesCollectionView.Get(), MeterTileViewWeak = MeterTileView.ToWeakPtr()]()
+			{
+				if (TSharedPtr<SLoudnessMetricTileView> PinnedMeterTileView = MeterTileViewWeak.Pin())
+				{
+					if (PinnedMeterTileView->HasDragDropOperation())
+					{
+						// Never hide the numeric values panel if a meter is being dragged so that the meter can be dropped onto it.
+						return EVisibility::Visible;
+					}
+				}
+				return (InValuesCollectionView.HasVisibleItems()) ? EVisibility::Visible : EVisibility::Collapsed;
+			});
+
 		return SNew(SBorder)
 			.Padding(0.0f)
 			.BorderImage(&BackgroundBrush)
@@ -775,14 +831,16 @@ namespace AudioWidgets
 					.AutoHeight()
 					[
 						SNew(SLoudnessMetricTileView,
-								FOnMakeTileInnerContent::CreateStatic(&MakeValueTileInnerContent, LoudnessScaleParams))
+								FOnMakeTileInnerContent::CreateStatic(&MakeValueTileInnerContent, LoudnessScaleParams),
+								SLoudnessMetricTileView::FOnSetShowItem::CreateStatic(&SetShowLoudnessMetricValue),
+								SLoudnessMetricTileView::FOnReorderItems::CreateSP(ValuesCollectionView, &FLoudnessMetricCollectionView::DraggedItemReorder))
 							.Visibility(GetNumericValuesPanelVisibility)
 							.OnGenerateTile_Static(&MakeNumericValueTileWidget)
 							.ListItemsSource(ValuesCollectionView->GetVisibleItems())
 							.ItemHeight(60.0f)
 							.ItemWidth(60.0f)
 							.ItemAlignment(EListItemAlignment::CenterAligned)
-							.SelectionMode(ESelectionMode::None)
+							.SelectionMode(ESelectionMode::Single)
 					]
 					+ SVerticalBox::Slot()
 					.AutoHeight()
@@ -869,6 +927,41 @@ namespace AudioWidgets
 		}
 	}
 
+	void FLoudnessMeterWidgetView::FLoudnessMetricCollectionView::DraggedItemReorder(const FLoudnessMetricRef& ItemDroppedOnto, bool bInsertAfter, const FLoudnessMetricRef& DragDropItem)
+	{
+		// Find the given items in the array of all loudness metrics (these are unchanging indexes that uniquely identify the loudness metrics):
+		const int32 DroppedOntoLoudnessMetricIndex = LoudnessMetrics->Find(ItemDroppedOnto);
+		const int32 DragDropLoudnessMetricIndex = LoudnessMetrics->Find(DragDropItem);
+		if (ensure(DroppedOntoLoudnessMetricIndex != INDEX_NONE && DragDropLoudnessMetricIndex != INDEX_NONE && DragDropLoudnessMetricIndex != DroppedOntoLoudnessMetricIndex))
+		{
+			// We get the current PermutationArray and then modify it as required.
+			const uint64 CurrentPermutationIndex = DisplayOrderParams.PermutationIndex.Get();
+			TArray<uint8> PermutationArray = GetPermutationArray(CurrentPermutationIndex, LoudnessMetrics->Num());
+
+			// Remove dragged item from current position:
+			PermutationArray.Remove(DragDropLoudnessMetricIndex);
+
+			const int32 DroppedOntoPosition = PermutationArray.Find(DroppedOntoLoudnessMetricIndex);
+			if (ensure(DroppedOntoPosition != INDEX_NONE))
+			{
+				// Place dragged item at new position:
+				const int32 InsertPosition = bInsertAfter ? DroppedOntoPosition + 1 : DroppedOntoPosition;
+				PermutationArray.EmplaceAt(InsertPosition, DragDropLoudnessMetricIndex);
+
+				// Get the new permutation index with the reordering applied:
+				const uint64 NewPermutationIndex = GetPermutationIndex(PermutationArray);
+				if (NewPermutationIndex != CurrentPermutationIndex && DisplayOrderParams.OnPermutationChanged.IsBound())
+				{
+					// Raise request to change ordering with new PermutationIndex:
+					DisplayOrderParams.OnPermutationChanged.Execute(NewPermutationIndex);
+
+					// Proactive refresh:
+					RefreshVisibleItems();
+				}
+			}
+		}
+	}
+
 	uint64 FLoudnessMeterWidgetView::FLoudnessMetricCollectionView::GetPermutationIndex(TArrayView<uint8> PermutationArray)
 	{
 		using namespace Permutations;
diff --git a/Engine/Plugins/Runtime/AudioWidgets/Source/AudioWidgets/Private/LoudnessMeterWidgetView.h b/Engine/Plugins/Runtime/AudioWidgets/Source/AudioWidgets/Private/LoudnessMeterWidgetView.h
--- a/Engine/Plugins/Runtime/AudioWidgets/Source/AudioWidgets/Private/LoudnessMeterWidgetView.h
+++ b/Engine/Plugins/Runtime/AudioWidgets/Source/AudioWidgets/Private/LoudnessMeterWidgetView.h
@@ -98,6 +98,8 @@ namespace AudioWidgets
 			TSharedRef<FLoudnessMetricRefArray> GetVisibleItems() { return VisibleItems; }
 			bool HasVisibleItems() const { return !VisibleItems->IsEmpty(); }
 
+			void DraggedItemReorder(const FLoudnessMetricRef& ItemDroppedOnto, bool bInsertAfter, const FLoudnessMetricRef& DragDropItem);
+
 		private:
 			static uint64 GetPermutationIndex(TArrayView<uint8> PermutationArray);
 			static TArray<uint8> GetPermutationArray(uint64 PermutationIndex, uint8 NumValues);
diff --git a/Engine/Plugins/Runtime/AudioWidgets/Source/AudioWidgets/Private/SDragReorderableTileView.cpp b/Engine/Plugins/Runtime/AudioWidgets/Source/AudioWidgets/Private/SDragReorderableTileView.cpp
new file mode 100644
--- /dev/null
+++ b/Engine/Plugins/Runtime/AudioWidgets/Source/AudioWidgets/Private/SDragReorderableTileView.cpp
@@ -0,0 +1,45 @@
+// Copyright Epic Games, Inc. All Rights Reserved.
+
+#include "SDragReorderableTileView.h"
+
+#include "Styling/StyleColors.h"
+
+namespace AudioWidgets
+{
+	int32 DrawDropIndicator(const FGeometry& AllottedGeometry, FSlateWindowElementList& OutDrawElements, int32 LayerId, const FWidgetStyle& InWidgetStyle, EHorizontalAlignment HorizontalAlignment)
+	{
+		const FVector2f TileSize = AllottedGeometry.GetLocalSize();
+
+		const float DropIndicatorTranslateX = [TileWidth = TileSize.X, HorizontalAlignment]()
+			{
+				switch (HorizontalAlignment)
+				{
+				default:
+				case HAlign_Fill:
+				case HAlign_Left:
+					return 0.0f;
+				case HAlign_Center:
+					return 0.5f * TileWidth;
+				case HAlign_Right:
+					return TileWidth;
+				}
+			}();
+
+		constexpr float DropIndicatorScale = 1.0f;
+		const FVector2f DropIndicatorSize(1.0f, TileSize.Y);
+		const FVector2f DropIndicatorTranslation(DropIndicatorTranslateX, 0.0f);
+		const FGeometry DropIndicatorGeometry = AllottedGeometry.MakeChild(DropIndicatorSize, FSlateLayoutTransform(DropIndicatorScale, DropIndicatorTranslation));
+
+		FSlateDrawElement::MakeBox
+		(
+			OutDrawElements,
+			LayerId++,
+			DropIndicatorGeometry.ToPaintGeometry(),
+			FAppStyle::GetBrush("PropertyEditor.VerticalDottedLine"),
+			ESlateDrawEffect::None,
+			FStyleColors::Select.GetColor(InWidgetStyle) * InWidgetStyle.GetColorAndOpacityTint()
+		);
+
+		return LayerId;
+	}
+}
diff --git a/Engine/Plugins/Runtime/AudioWidgets/Source/AudioWidgets/Private/SDragReorderableTileView.h b/Engine/Plugins/Runtime/AudioWidgets/Source/AudioWidgets/Private/SDragReorderableTileView.h
new file mode 100644
--- /dev/null
+++ b/Engine/Plugins/Runtime/AudioWidgets/Source/AudioWidgets/Private/SDragReorderableTileView.h
@@ -0,0 +1,304 @@
+// Copyright Epic Games, Inc. All Rights Reserved.
+
+#pragma once
+
+#include "Framework/SlateDelegates.h"
+#include "Widgets/Views/STileView.h"
+
+namespace AudioWidgets
+{
+	/**
+	 * Function to draw a dotted line representing where the current drag drop would place the dragged item.
+	 */
+	int32 DrawDropIndicator(const FGeometry& AllottedGeometry, FSlateWindowElementList& OutDrawElements, int32 LayerId, const FWidgetStyle& InWidgetStyle, EHorizontalAlignment HorizontalAlignment);
+
+	template<typename DragDropOperationT>
+	class SDragReorderableTileView;
+
+	/**
+	 * Abstract base class, represents an item that is being dragged from a given source tile view.
+	 * Uses Curiously Recurring Template Pattern so that FDragDropOperation type checking operations can be used.
+	 */
+	template<typename ItemT, typename DragDropOperationT>
+	class TItemDragDropOperation : public FDragDropOperation
+	{
+	public:
+		using ItemType = ItemT;
+		using STileViewType = SDragReorderableTileView<DragDropOperationT>;
+
+		TItemDragDropOperation(TSharedRef<STileViewType> InSourceTileView)
+			: WeakSourceTileView(InSourceTileView.ToWeakPtr())
+		{
+		}
+
+		virtual const ItemType& GetItem() const = 0;
+
+		TSharedPtr<STileViewType> GetSourceTileView() const { return WeakSourceTileView.Pin(); }
+
+		virtual FCursorReply OnCursorQuery() override { return FCursorReply::Cursor(EMouseCursor::GrabHand); }
+
+	private:
+		const TWeakPtr<STileViewType> WeakSourceTileView;
+	};
+
+	template<typename T>
+	concept ItemDragDropOperation =
+		requires { typename T::ItemType; };// && std::derived_from<T, TItemDragDropOperation<typename T::ItemType, T>>; // Can uncomment when all platform stdlibs have <concepts> header.
+
+
+	/**
+	 * Abstract base class for TileView widget that allow its tiles to be dragged and dropped for reordering etc.
+	 * Dropping of a dragged item onto the tile view is supported if the tile view is currently empty.
+	 */
+	template<typename DragDropOperationT>
+	class SDragReorderableTileView : public STileView<typename DragDropOperationT::ItemType>
+	{
+	public:
+		using Base = STileView<typename DragDropOperationT::ItemType>;
+		using ItemType = typename DragDropOperationT::ItemType;
+
+		DECLARE_DELEGATE_TwoParams(FOnSetShowItem, const ItemType&, bool);
+		DECLARE_DELEGATE_ThreeParams(FOnReorderItems, const ItemType&/*ItemDroppedOnto*/, bool/*bInsertAfter*/, const ItemType&/*DragDropItem*/);
+
+		void Construct(const Base::FArguments& InArgs, FOnSetShowItem InOnSetShowItem, FOnReorderItems InOnReorderItems)
+		{
+			Base::Construct(InArgs);
+
+			OnSetShowItem = InOnSetShowItem;
+			OnReorderItems = InOnReorderItems;
+		}
+
+		virtual TSharedPtr<DragDropOperationT> CreateDragDropOperation() = 0;
+
+		void SetShowItem(const ItemType& Item, bool bShow)
+		{
+			OnSetShowItem.ExecuteIfBound(Item, bShow);
+		}
+
+		void ReorderItems(const ItemType& ItemDroppedOnto, bool bInsertAfter, const ItemType& DragDropItem)
+		{
+			ensure(ItemDroppedOnto != DragDropItem);
+			OnReorderItems.ExecuteIfBound(ItemDroppedOnto, bInsertAfter, DragDropItem);
+		}
+
+		// Begin SWidget overrides.
+		virtual FReply OnDragOver(const FGeometry& MyGeometry, const FDragDropEvent& DragDropEvent) override
+		{
+			if (TSharedPtr<DragDropOperationT> DragDropOp = DragDropEvent.GetOperationAs<DragDropOperationT>())
+			{
+				// Respond to drag drop only if we have no items, otherwise the tile will handle it:
+				if (Base::GetItems().IsEmpty())
+				{
+					bPaintDropIndicator = true;
+					return FReply::Handled();
+				}
+			}
+			return FReply::Unhandled();
+		}
+
+		virtual void OnDragLeave(FDragDropEvent const& DragDropEvent) override
+		{
+			bPaintDropIndicator = false;
+		}
+
+		virtual FReply OnDrop(const FGeometry& MyGeometry, const FDragDropEvent& DragDropEvent) override
+		{
+			bPaintDropIndicator = false;
+
+			if (TSharedPtr<DragDropOperationT> DragDropOp = DragDropEvent.GetOperationAs<DragDropOperationT>())
+			{
+				// Respond to drag drop only if we have no items, otherwise the tile will handle it:
+				if (Base::GetItems().IsEmpty())
+				{
+					if (TSharedPtr<SDragReorderableTileView> DragSourceTileView = DragDropOp->GetSourceTileView())
+					{
+						const auto& DragDropItem = DragDropOp->GetItem();
+
+						// Move between tile views:
+						DragSourceTileView->SetShowItem(DragDropItem, false);
+						SetShowItem(DragDropItem, true);
+
+						return FReply::Handled();
+					}
+				}
+			}
+			return FReply::Unhandled();
+		}
+
+		virtual int32 OnPaint(const FPaintArgs& Args, const FGeometry& AllottedGeometry, const FSlateRect& MyCullingRect, FSlateWindowElementList& OutDrawElements, int32 LayerId, const FWidgetStyle& InWidgetStyle, bool bParentEnabled) const override
+		{
+			LayerId = Base::OnPaint(Args, AllottedGeometry, MyCullingRect, OutDrawElements, LayerId, InWidgetStyle, bParentEnabled);
+
+			if (bPaintDropIndicator)
+			{
+				LayerId = DrawDropIndicator(AllottedGeometry, OutDrawElements, LayerId, InWidgetStyle, EHorizontalAlignment::HAlign_Center);
+			}
+
+			return LayerId;
+		}
+		// End SWidget overrides
+
+	private:
+		FOnSetShowItem OnSetShowItem;
+		FOnReorderItems OnReorderItems;
+		bool bPaintDropIndicator = false;
+	};
+
+	/**
+	 * Base class for tiles displayed by SDragReorderableTileView.
+	 */
+	template<ItemDragDropOperation DragDropOperationT>
+	class SDragReorderableTile : public STableRow<typename DragDropOperationT::ItemType>
+	{
+	public:
+		using Base = STableRow<typename DragDropOperationT::ItemType>;
+		using ItemType = typename DragDropOperationT::ItemType;
+		using STileViewType = SDragReorderableTileView<DragDropOperationT>;
+
+		bool IsBeingDragged() const { return WeakDragDropOperation.IsValid(); }
+
+	protected:
+		// Begin SWidget overrides.
+		virtual FReply OnDragDetected(const FGeometry& MyGeometry, const FPointerEvent& PointerEvent) override
+		{
+			const FReply BaseReply = Base::OnDragDetected(MyGeometry, PointerEvent);
+			if (!BaseReply.IsEventHandled() && PointerEvent.IsMouseButtonDown(EKeys::LeftMouseButton))
+			{
+				if (TSharedPtr<STileViewType> OwnerTileView = GetPinnedTileView())
+				{
+					if (TSharedPtr<FDragDropOperation> DragDropOperation = OwnerTileView->CreateDragDropOperation())
+					{
+						// Hold a weak pointer to maintain awareness of whether this tile is currently being dragged:
+						WeakDragDropOperation = DragDropOperation;
+
+						return FReply::Handled().BeginDragDrop(DragDropOperation.ToSharedRef());
+					}
+				}
+			}
+			return BaseReply;
+		}
+
+		virtual FReply OnDragOver(const FGeometry& MyGeometry, const FDragDropEvent& DragDropEvent) override
+		{
+			// Setting STableRow::ItemDropZone TOptional will enable painting the drop indicator:
+			Base::ItemDropZone = CanAcceptDrop(MyGeometry, DragDropEvent);
+
+			return FReply::Handled();
+		}
+
+		virtual FReply OnDrop(const FGeometry& MyGeometry, const FDragDropEvent& DragDropEvent) override
+		{
+			// Clear STableRow::ItemDropZone TOptional to no longer paint the drop indicator:
+			Base::ItemDropZone.Reset();
+
+			if (TOptional<EItemDropZone> AcceptableDropZone = CanAcceptDrop(MyGeometry, DragDropEvent))
+			{
+				return AcceptDrop(DragDropEvent, *AcceptableDropZone);
+			}
+
+			return FReply::Unhandled();
+		}
+		// End SWidget overrides
+
+		// Begin STableRow overrides.
+		virtual int32 OnPaintDropIndicator(EItemDropZone InItemDropZone, const FPaintArgs& Args, const FGeometry& AllottedGeometry, const FSlateRect& MyCullingRect, FSlateWindowElementList& OutDrawElements, int32 LayerId, const FWidgetStyle& InWidgetStyle, bool bParentEnabled) const override
+		{
+			const EHorizontalAlignment HorizontalAlignment = (InItemDropZone == EItemDropZone::BelowItem) ? EHorizontalAlignment::HAlign_Right : EHorizontalAlignment::HAlign_Left;
+			return DrawDropIndicator(AllottedGeometry, OutDrawElements, LayerId, InWidgetStyle, HorizontalAlignment);
+		}
+		// End STableRow overrides.
+
+	private:
+		TOptional<EItemDropZone> CanAcceptDrop(const FGeometry& MyGeometry, const FDragDropEvent& DragDropEvent)
+		{
+			if (TSharedPtr<DragDropOperationT> DragDropOp = DragDropEvent.GetOperationAs<DragDropOperationT>())
+			{
+				if (TSharedPtr<STileViewType> OwnerTileView = GetPinnedTileView())
+				{
+					const EItemDropZone ItemHoverZone = GetItemHoverZone(MyGeometry, DragDropEvent);
+					if (IsValidDrop(OwnerTileView.ToSharedRef(), ItemHoverZone, DragDropOp.ToSharedRef()))
+					{
+						return ItemHoverZone;
+					}
+				}
+			}
+			return NullOpt;
+		}
+
+		bool IsValidDrop(const TSharedRef<STileViewType>& OwnerTileView, EItemDropZone ItemHoverZone, TSharedRef<DragDropOperationT> DragDropOp) const
+		{
+			const auto& DragDropItem = DragDropOp->GetItem();
+
+			TConstArrayView<ItemType> Items = OwnerTileView->GetItems();
+			const int32 DragDroppedIndex = Items.Find(DragDropItem);
+			if (DragDroppedIndex == INDEX_NONE)
+			{
+				// Item does not exist in the current tile view so it is always valid.
+				return true;
+			}
+
+			if (DragDropOp->GetSourceTileView() != OwnerTileView)
+			{
+				// Don't alow drag from another tile view as the current tile view already contains the loudness metric.
+				return false;
+			}
+
+			if (const ItemType* ItemPtr = Base::GetItemForThis(OwnerTileView))
+			{
+				const int32 DroppedOntoIndex = Items.Find(*ItemPtr);
+				if (DroppedOntoIndex != INDEX_NONE && DragDroppedIndex != DroppedOntoIndex)
+				{
+					const bool bInsertAfter = (ItemHoverZone == EItemDropZone::BelowItem);
+					const int32 InsertIndex = (bInsertAfter) ? DroppedOntoIndex + 1 : DroppedOntoIndex - 1;
+
+					// Return valid if insert position causes change:
+					return (InsertIndex != DragDroppedIndex);
+				}
+			}
+			return false;
+		}
+
+		FReply AcceptDrop(const FDragDropEvent& DragDropEvent, EItemDropZone ItemHoverZone)
+		{
+			if (TSharedPtr<DragDropOperationT> DragDropOp = DragDropEvent.GetOperationAs<DragDropOperationT>())
+			{
+				if (TSharedPtr<STileViewType> OwnerTileView = GetPinnedTileView())
+				{
+					if (const ItemType* ItemPtr = Base::GetItemForThis(OwnerTileView.ToSharedRef()))
+					{
+						const auto& DragDropItem = DragDropOp->GetItem();
+
+						// Place the dropped item as implied by ItemHoverZone:
+						const bool bInsertAfter = (ItemHoverZone == EItemDropZone::BelowItem);
+						OwnerTileView->ClearSelection();
+						OwnerTileView->ReorderItems(*ItemPtr, bInsertAfter, DragDropItem);
+
+						TSharedPtr<STileViewType> DragSourceTileView = DragDropOp->GetSourceTileView();
+						if (DragSourceTileView.IsValid() && DragSourceTileView != OwnerTileView)
+						{
+							// Item was dragged from other tile view, move it between tile views:
+							DragSourceTileView->SetShowItem(DragDropItem, false);
+							OwnerTileView->SetShowItem(DragDropItem, true);
+						}
+
+						return FReply::Handled();
+					}
+				}
+			}
+			return FReply::Unhandled();
+		}
+
+		TSharedPtr<STileViewType> GetPinnedTileView()
+		{
+			return StaticCastSharedPtr<STileViewType>(Base::OwnerTablePtr.Pin());
+		}
+
+		EItemDropZone GetItemHoverZone(const FGeometry& MyGeometry, const FDragDropEvent& DragDropEvent)
+		{
+			const FVector2f LocalPointerPos = MyGeometry.AbsoluteToLocal(DragDropEvent.GetScreenSpacePosition());
+			return Base::ZoneFromPointerPosition(LocalPointerPos, MyGeometry.GetLocalSize(), EOrientation::Orient_Horizontal);
+		}
+
+		TWeakPtr<FDragDropOperation> WeakDragDropOperation; // Will be set when a drag is detected.
+	};
+}
