Fix compile error - missing "typename" codeword in parameter declaration

#Synced-CL 49317289

diff --git a/Engine/Plugins/Runtime/AudioWidgets/Source/AudioWidgets/Private/SDragReorderableTileView.h b/Engine/Plugins/Runtime/AudioWidgets/Source/AudioWidgets/Private/SDragReorderableTileView.h
--- a/Engine/Plugins/Runtime/AudioWidgets/Source/AudioWidgets/Private/SDragReorderableTileView.h
+++ b/Engine/Plugins/Runtime/AudioWidgets/Source/AudioWidgets/Private/SDragReorderableTileView.h
@@ -60,7 +60,7 @@ namespace AudioWidgets
 		DECLARE_DELEGATE_TwoParams(FOnSetShowItem, const ItemType&, bool);
 		DECLARE_DELEGATE_ThreeParams(FOnReorderItems, const ItemType&/*ItemDroppedOnto*/, bool/*bInsertAfter*/, const ItemType&/*DragDropItem*/);
 
-		void Construct(const Base::FArguments& InArgs, FOnSetShowItem InOnSetShowItem, FOnReorderItems InOnReorderItems)
+		void Construct(const typename Base::FArguments& InArgs, FOnSetShowItem InOnSetShowItem, FOnReorderItems InOnReorderItems)
 		{
 			Base::Construct(InArgs);
 
