[AudioWidgets] Added TruePeak meter so we are able to visualize it in the collection of Loudness meters used in MetaSound editor and Audio Insights.

#Synced-CL 50481123

diff --git a/Engine/Plugins/Runtime/AudioSynesthesia/Source/AudioSynesthesia/Classes/Meter.h b/Engine/Plugins/Runtime/AudioSynesthesia/Source/AudioSynesthesia/Classes/Meter.h
--- a/Engine/Plugins/Runtime/AudioSynesthesia/Source/AudioSynesthesia/Classes/Meter.h
+++ b/Engine/Plugins/Runtime/AudioSynesthesia/Source/AudioSynesthesia/Classes/Meter.h
@@ -15,6 +15,7 @@ enum class EMeterPeakType : uint8
 	MeanSquared,
 	RootMeanSquared,
 	Peak,
+	TruePeak,
 	Count UMETA(Hidden)
 };
 
@@ -54,6 +55,10 @@ public:
 	UPROPERTY(EditAnywhere, BlueprintReadOnly, Category = AudioAnalyzer, meta = (ClampMin = "0.0"))
 	float ClippingThreshold = 1.0f;
 
+	/** Analog response. */
+	UPROPERTY(EditAnywhere, BlueprintReadOnly, Category = AudioAnalyzer)
+	bool bIsAnalog = true;
+
 	/** Convert UMeterSettings to IAnalyzerSettings */
 	UE_API TUniquePtr<Audio::IAnalyzerSettings> GetSettings(const int32 InSampleRate, const int32 InNumChannels) const;
 
diff --git a/Engine/Plugins/Runtime/AudioSynesthesia/Source/AudioSynesthesia/Private/Meter.cpp b/Engine/Plugins/Runtime/AudioSynesthesia/Source/AudioSynesthesia/Private/Meter.cpp
--- a/Engine/Plugins/Runtime/AudioSynesthesia/Source/AudioSynesthesia/Private/Meter.cpp
+++ b/Engine/Plugins/Runtime/AudioSynesthesia/Source/AudioSynesthesia/Private/Meter.cpp
@@ -21,6 +21,10 @@ TUniquePtr<Audio::IAnalyzerSettings> UMeterSettings::GetSettings(const int32 InS
 		Settings->MeterPeakMode = Audio::EPeakMode::RootMeanSquared;
 		break;
 
+	case EMeterPeakType::TruePeak:
+		Settings->MeterPeakMode = Audio::EPeakMode::TruePeak;
+		break;
+
 	default:
 	case EMeterPeakType::Peak:
 		Settings->MeterPeakMode = Audio::EPeakMode::Peak;
@@ -31,6 +35,7 @@ TUniquePtr<Audio::IAnalyzerSettings> UMeterSettings::GetSettings(const int32 InS
 	Settings->MeterReleaseTime = MeterReleaseTime;
 	Settings->PeakHoldTime = PeakHoldTime;
 	Settings->ClippingThreshold = ClippingThreshold;
+	Settings->bIsAnalog = bIsAnalog;
 
 	return Settings;
 }
diff --git a/Engine/Plugins/Runtime/AudioSynesthesia/Source/AudioSynesthesiaCore/Private/MeterAnalyzer.cpp b/Engine/Plugins/Runtime/AudioSynesthesia/Source/AudioSynesthesiaCore/Private/MeterAnalyzer.cpp
--- a/Engine/Plugins/Runtime/AudioSynesthesia/Source/AudioSynesthesiaCore/Private/MeterAnalyzer.cpp
+++ b/Engine/Plugins/Runtime/AudioSynesthesia/Source/AudioSynesthesiaCore/Private/MeterAnalyzer.cpp
@@ -7,7 +7,7 @@
 namespace Audio
 {
 	FMeterAnalyzer::FMeterAnalyzer(float InSampleRate, int32 InNumChannels, const FMeterAnalyzerSettings& InSettings)
-	: EnvelopeFollower(FEnvelopeFollowerInitParams{InSampleRate, InNumChannels, static_cast<float>(InSettings.MeterAttackTime), static_cast<float>(InSettings.MeterReleaseTime), InSettings.MeterPeakMode, true /* bInIsAnalog */})
+	: EnvelopeFollower(FEnvelopeFollowerInitParams{InSampleRate, InNumChannels, static_cast<float>(InSettings.MeterAttackTime), static_cast<float>(InSettings.MeterReleaseTime), InSettings.MeterPeakMode, InSettings.bIsAnalog })
 	, Settings(InSettings)
 	, SampleRate(InSampleRate)
 	, NumChannels(InNumChannels)
diff --git a/Engine/Plugins/Runtime/AudioSynesthesia/Source/AudioSynesthesiaCore/Public/MeterAnalyzer.h b/Engine/Plugins/Runtime/AudioSynesthesia/Source/AudioSynesthesiaCore/Public/MeterAnalyzer.h
--- a/Engine/Plugins/Runtime/AudioSynesthesia/Source/AudioSynesthesiaCore/Public/MeterAnalyzer.h
+++ b/Engine/Plugins/Runtime/AudioSynesthesia/Source/AudioSynesthesiaCore/Public/MeterAnalyzer.h
@@ -24,6 +24,8 @@ namespace Audio
 
 		// The volume threshold to detect clipping
 		float ClippingThreshold = 1.0f;
+
+		bool bIsAnalog = true;
 	};
 
 	// Per-channel analyzer results
diff --git a/Engine/Plugins/Runtime/AudioWidgets/Source/AudioWidgets/Private/AudioLoudnessMeter.cpp b/Engine/Plugins/Runtime/AudioWidgets/Source/AudioWidgets/Private/AudioLoudnessMeter.cpp
--- a/Engine/Plugins/Runtime/AudioWidgets/Source/AudioWidgets/Private/AudioLoudnessMeter.cpp
+++ b/Engine/Plugins/Runtime/AudioWidgets/Source/AudioWidgets/Private/AudioLoudnessMeter.cpp
@@ -2,6 +2,7 @@
 
 #include "AudioLoudnessMeter.h"
 
+#include "AudioMeterTypes.h"
 #include "LoudnessMeterRackUnitSettings.h"
 #include "Widgets/Docking/SDockTab.h"
 
@@ -102,6 +103,11 @@ namespace AudioWidgets
 		LKFSSettings = TStrongObjectPtr(NewObject<ULKFSSettings>());
 		LKFSSettings->bCalculateOverallLoudnessRange = true;
 
+		TruePeakSettings = TStrongObjectPtr(NewObject<UMeterSettings>());
+		TruePeakSettings->PeakMode = EMeterPeakType::TruePeak;
+		TruePeakSettings->bIsAnalog = false;
+		TruePeakSettings->MeterAttackTime = 0;
+
 		AudioDeviceId = InAudioDeviceId;
 
 		// Only create analyzers etc if we have an audio device:
@@ -114,23 +120,38 @@ namespace AudioWidgets
 			AudioBus->AudioBusChannels = EAudioBusChannels(InNumChannels - 1);
 
 			CreateLKFSAnalyzer();
+			CreateTruePeakAnalyzer();
 
 			LKFSAnalyzer->StartAnalyzing(InAudioDeviceId, AudioBus.Get());
+			TruePeakAnalyzer->StartAnalyzing(InAudioDeviceId, AudioBus.Get());
 		}
 	}
 
 	void FAudioLoudnessMeter::Teardown()
 	{
-		if (UObjectInitialized() && LKFSAnalyzer.IsValid() && LKFSAnalyzer->IsValidLowLevel())
+		if (UObjectInitialized())
 		{
-			LKFSAnalyzer->StopAnalyzing();
-			ReleaseLKFSAnalyzer();
+			if (TruePeakAnalyzer.IsValid() && TruePeakAnalyzer->IsValidLowLevel())
+			{
+				TruePeakAnalyzer->StopAnalyzing();
+				ReleaseTruePeakAnalyzer();
+			}
+
+			if (LKFSAnalyzer.IsValid() && LKFSAnalyzer->IsValidLowLevel())
+			{
+				LKFSAnalyzer->StopAnalyzing();
+				ReleaseLKFSAnalyzer();
+			}
 		}
 
+		LatestOverallTruePeakMeterResultsDelegateHandle.Reset();
+		TruePeakAnalyzer.Reset();
+
 		LatestOverallLKFSResultsDelegateHandle.Reset();
 		LKFSAnalyzer.Reset();
 
 		AudioBus.Reset();
+		TruePeakSettings.Reset();
 		LKFSSettings.Reset();
 		bUseExternalAudioBus = false;
 	}
@@ -156,6 +177,65 @@ namespace AudioWidgets
 		LKFSAnalyzer.Reset();
 	}
 
+	void FAudioLoudnessMeter::ResetLKFSAnalyzer()
+	{
+		// Reset loudness analyzer
+		if (LKFSAnalyzer.IsValid())
+		{
+			LKFSAnalyzer->StopAnalyzing();
+			ReleaseLKFSAnalyzer();
+		}
+
+		LatestLoudnessResults.Reset();
+
+		// Recreate loudness analyzer
+		if (AudioDeviceId != FAudioBusInfo::InvalidAudioDeviceId)
+		{
+			CreateLKFSAnalyzer();
+			LKFSAnalyzer->StartAnalyzing(AudioDeviceId, AudioBus.Get());
+		}
+	}
+
+	void FAudioLoudnessMeter::CreateTruePeakAnalyzer()
+	{
+		ensure(!TruePeakAnalyzer.IsValid());
+		ensure(!LatestOverallTruePeakMeterResultsDelegateHandle.IsValid());
+
+		TruePeakAnalyzer = TStrongObjectPtr(NewObject<UMeterAnalyzer>());
+		TruePeakAnalyzer->Settings = TruePeakSettings.Get();
+		LatestOverallTruePeakMeterResultsDelegateHandle = TruePeakAnalyzer->OnLatestOverallMeterResultsNative.AddRaw(this, &FAudioLoudnessMeter::OnTruePeakOutput);
+	}
+
+	void FAudioLoudnessMeter::ReleaseTruePeakAnalyzer()
+	{
+		if (ensure(TruePeakAnalyzer.IsValid() && LatestOverallTruePeakMeterResultsDelegateHandle.IsValid()))
+		{
+			TruePeakAnalyzer->OnLatestOverallMeterResultsNative.Remove(LatestOverallTruePeakMeterResultsDelegateHandle);
+		}
+
+		LatestOverallTruePeakMeterResultsDelegateHandle.Reset();
+		TruePeakAnalyzer.Reset();
+	}
+
+	void FAudioLoudnessMeter::ResetTruePeakAnalyzer()
+	{
+		// Reset true peak analyzer
+		if (TruePeakAnalyzer.IsValid())
+		{
+			TruePeakAnalyzer->StopAnalyzing();
+			ReleaseTruePeakAnalyzer();
+		}
+
+		LatestTruePeakResults.Reset();
+
+		// Recreate true peak analyzer
+		if (AudioDeviceId != FAudioBusInfo::InvalidAudioDeviceId)
+		{
+			CreateTruePeakAnalyzer();
+			TruePeakAnalyzer->StartAnalyzing(AudioDeviceId, AudioBus.Get());
+		}
+	}
+
 	void FAudioLoudnessMeter::InitLoudnessMeterWidgetView(const FAudioAnalyzerRackUnitConstructParams& Params)
 	{
 		// Initializing these callbacks outside of the FAudioLoudnessMeter constructor allows us to use the CreateSP(...) smart pointer variants for attributes/delegates:
@@ -176,6 +256,7 @@ namespace AudioWidgets
 			{
 				.Name = GET_MEMBER_NAME_CHECKED(FLoudnessMeterRackUnitSettings, LongTermLoudness),
 				.DisplayName = LOCTEXT("LongTermLoudnessDisplayName", "Long Term"),
+				.MeterMetric = EAudioMeterMetric::Loudness,
 				.Value = TAttribute<TOptional<float>>::CreateSPLambda(this, [this]() { return (LatestLoudnessResults.IsSet()) ? TOptional<float>(LatestLoudnessResults->GatedLoudness) : NullOpt; }),
 				.bShowValue = true,
 				.bShowMeter = false,
@@ -183,6 +264,7 @@ namespace AudioWidgets
 			{
 				.Name = GET_MEMBER_NAME_CHECKED(FLoudnessMeterRackUnitSettings, ShortTermLoudness),
 				.DisplayName = LOCTEXT("ShortTermLoudnessDisplayName", "Short Term"),
+				.MeterMetric = EAudioMeterMetric::Loudness,
 				.Value = TAttribute<TOptional<float>>::CreateSPLambda(this, [this]() { return (LatestLoudnessResults.IsSet()) ? TOptional<float>(LatestLoudnessResults->ShortTermLoudness) : NullOpt; }),
 				.bShowValue = false,
 				.bShowMeter = true,
@@ -190,6 +272,7 @@ namespace AudioWidgets
 			{
 				.Name = GET_MEMBER_NAME_CHECKED(FLoudnessMeterRackUnitSettings, MomentaryLoudness),
 				.DisplayName = LOCTEXT("MomentaryLoudnessDisplayName", "Momentary"),
+				.MeterMetric = EAudioMeterMetric::Loudness,
 				.Value = TAttribute<TOptional<float>>::CreateSPLambda(this, [this]() { return (LatestLoudnessResults.IsSet()) ? TOptional<float>(LatestLoudnessResults->Loudness) : NullOpt; }),
 				.bShowValue = false,
 				.bShowMeter = true,
@@ -197,9 +280,19 @@ namespace AudioWidgets
 			{
 				.Name = GET_MEMBER_NAME_CHECKED(FLoudnessMeterRackUnitSettings, LoudnessRange),
 				.DisplayName = LOCTEXT("LoudnessRangeDisplayName", "Range"),
+				.MeterMetric = EAudioMeterMetric::Loudness,
 				.Range = TAttribute<TOptional<FFloatInterval>>::CreateSPLambda(this, [this]() { return (LatestLoudnessResults.IsSet()) ? LatestLoudnessResults->LoudnessRange : NullOpt; }),
 				.bShowValue = true,
 				.bShowMeter = false,
+			},
+			{
+				.Name = GET_MEMBER_NAME_CHECKED(FLoudnessMeterRackUnitSettings, TruePeak),
+				.DisplayName = LOCTEXT("TruePeakDisplayName", "True Peak"),
+				.MeterMetric = EAudioMeterMetric::Decibels,
+				.Value = TAttribute<TOptional<float>>::CreateSPLambda(this, [this]() { return (LatestTruePeakResults.IsSet()) ? TOptional<float>(LatestTruePeakResults->MeterValue) : NullOpt; }),
+				.bShowValue = true,
+				.bShowMeter = false,
+				.OnClicked = FSimpleDelegate::CreateSPLambda(this, [this]() { ResetTruePeakAnalyzer(); })
 			}
 		};
 
@@ -219,24 +312,24 @@ namespace AudioWidgets
 					});
 
 				LoudnessScaleParams.Range.BindLambda([=]() { return SettingsHelper.GetRackUnitSettings().LoudnessScaleRange; });
-				LoudnessScaleParams.OnRangeValueChanged.BindLambda([=](int Value) { SettingsHelper.GetRackUnitSettings().LoudnessScaleRange = Value; });
-				LoudnessScaleParams.OnRangeValueCommitted.BindLambda([=](int Value, ETextCommit::Type CommitType)
+				LoudnessScaleParams.OnRangeValueChanged.BindLambda([=](int32 Value) { SettingsHelper.GetRackUnitSettings().LoudnessScaleRange = Value; });
+				LoudnessScaleParams.OnRangeValueCommitted.BindLambda([=](int32 Value, ETextCommit::Type CommitType)
 					{
 						SettingsHelper.GetRackUnitSettings().LoudnessScaleRange = Value;
 						SettingsHelper.SaveConfig();
 					});
 
 				LoudnessScaleParams.Offset.BindLambda([=]() { return SettingsHelper.GetRackUnitSettings().LoudnessScaleOffset; });
-				LoudnessScaleParams.OnOffsetValueChanged.BindLambda([=](int Value) { SettingsHelper.GetRackUnitSettings().LoudnessScaleOffset = Value; });
-				LoudnessScaleParams.OnOffsetValueCommitted.BindLambda([=](int Value, ETextCommit::Type CommitType)
+				LoudnessScaleParams.OnOffsetValueChanged.BindLambda([=](int32 Value) { SettingsHelper.GetRackUnitSettings().LoudnessScaleOffset = Value; });
+				LoudnessScaleParams.OnOffsetValueCommitted.BindLambda([=](int32 Value, ETextCommit::Type CommitType)
 					{
 						SettingsHelper.GetRackUnitSettings().LoudnessScaleOffset = Value;
 						SettingsHelper.SaveConfig();
 					});
 
 				LoudnessScaleParams.Target.BindLambda([=]() { return SettingsHelper.GetRackUnitSettings().LoudnessScaleTarget; });
-				LoudnessScaleParams.OnTargetValueChanged.BindLambda([=](int Value) { SettingsHelper.GetRackUnitSettings().LoudnessScaleTarget = Value; });
-				LoudnessScaleParams.OnTargetValueCommitted.BindLambda([=](int Value, ETextCommit::Type CommitType)
+				LoudnessScaleParams.OnTargetValueChanged.BindLambda([=](int32 Value) { SettingsHelper.GetRackUnitSettings().LoudnessScaleTarget = Value; });
+				LoudnessScaleParams.OnTargetValueCommitted.BindLambda([=](int32 Value, ETextCommit::Type CommitType)
 					{
 						SettingsHelper.GetRackUnitSettings().LoudnessScaleTarget = Value;
 						SettingsHelper.SaveConfig();
@@ -251,14 +344,14 @@ namespace AudioWidgets
 					});
 
 				ValuesDisplayOrderParams.PermutationIndex.BindLambda([=]() { return SettingsHelper.GetRackUnitSettings().ValuesOrderingPermutation; });
-				ValuesDisplayOrderParams.OnPermutationChanged.BindLambda([=](int NewPermutationIndex)
+				ValuesDisplayOrderParams.OnPermutationChanged.BindLambda([=](int32 NewPermutationIndex)
 					{
 						SettingsHelper.GetRackUnitSettings().ValuesOrderingPermutation = NewPermutationIndex;
 						SettingsHelper.SaveConfig();
 					});
 
 				MetersDisplayOrderParams.PermutationIndex.BindLambda([=]() { return SettingsHelper.GetRackUnitSettings().MetersOrderingPermutation; });
-				MetersDisplayOrderParams.OnPermutationChanged.BindLambda([=](int NewPermutationIndex)
+				MetersDisplayOrderParams.OnPermutationChanged.BindLambda([=](int32 NewPermutationIndex)
 					{
 						SettingsHelper.GetRackUnitSettings().MetersOrderingPermutation = NewPermutationIndex;
 						SettingsHelper.SaveConfig();
@@ -349,6 +442,25 @@ namespace AudioWidgets
 		}
 	}
 
+	void FAudioLoudnessMeter::OnTruePeakOutput(UMeterAnalyzer* InMeterAnalyzer, const FMeterResults& InMeterResults)
+	{
+		if (InMeterAnalyzer == TruePeakAnalyzer.Get())
+		{
+			if (!LatestTruePeakResults.IsSet())
+			{
+				LatestTruePeakResults = InMeterResults;
+			}
+			else
+			{
+				// Keep always the max true peak value (unless reset)
+				const float PreviousTruePeakValue = LatestTruePeakResults->MeterValue;
+
+				LatestTruePeakResults = InMeterResults;
+				LatestTruePeakResults->MeterValue = FMath::Max<float>(PreviousTruePeakValue, InMeterResults.MeterValue);
+			}
+		}
+	}
+
 	FTimespan FAudioLoudnessMeter::GetAnalysisTime() const
 	{
 		if (LatestLoudnessResults.IsSet() && LKFSSettings.IsValid())
@@ -362,19 +474,8 @@ namespace AudioWidgets
 
 	FReply FAudioLoudnessMeter::HandleResetButtonClicked()
 	{
-		if (LKFSAnalyzer.IsValid())
-		{
-			LKFSAnalyzer->StopAnalyzing();
-			ReleaseLKFSAnalyzer();
-		}
-
-		LatestLoudnessResults.Reset();
-
-		if (AudioDeviceId != FAudioBusInfo::InvalidAudioDeviceId)
-		{
-			CreateLKFSAnalyzer();
-			LKFSAnalyzer->StartAnalyzing(AudioDeviceId, AudioBus.Get());
-		}
+		ResetLKFSAnalyzer();
+		ResetTruePeakAnalyzer();
 
 		return FReply::Handled();
 	}
diff --git a/Engine/Plugins/Runtime/AudioWidgets/Source/AudioWidgets/Private/AudioLoudnessMeter.h b/Engine/Plugins/Runtime/AudioWidgets/Source/AudioWidgets/Private/AudioLoudnessMeter.h
--- a/Engine/Plugins/Runtime/AudioWidgets/Source/AudioWidgets/Private/AudioLoudnessMeter.h
+++ b/Engine/Plugins/Runtime/AudioWidgets/Source/AudioWidgets/Private/AudioLoudnessMeter.h
@@ -5,6 +5,7 @@
 #include "Containers/Ticker.h"
 #include "LKFS.h"
 #include "LoudnessMeterWidgetView.h"
+#include "Meter.h"
 
 namespace AudioWidgets
 {
@@ -43,10 +44,16 @@ namespace AudioWidgets
 
 		void CreateLKFSAnalyzer();
 		void ReleaseLKFSAnalyzer();
+		void ResetLKFSAnalyzer();
+
+		void CreateTruePeakAnalyzer();
+		void ReleaseTruePeakAnalyzer();
+		void ResetTruePeakAnalyzer();
 
 		void InitLoudnessMeterWidgetView(const FAudioAnalyzerRackUnitConstructParams& Params);
 
 		void OnLoudnessOutput(ULKFSAnalyzer* InLKFSAnalyzer, const FLKFSResults& InLoudnessResults);
+		void OnTruePeakOutput(UMeterAnalyzer* InMeterAnalyzer, const FMeterResults& InMeterResults);
 
 		FTimespan GetAnalysisTime() const;
 		FReply HandleResetButtonClicked();
@@ -54,18 +61,30 @@ namespace AudioWidgets
 		/** Creates the meters widget */
 		FLoudnessMeterWidgetView LoudnessMeterWidgetView;
 		
-		/** Audio analyzer object. */
+		/** Loudness analyzer object. */
 		TStrongObjectPtr<ULKFSAnalyzer> LKFSAnalyzer;
 
-		/** Analyzer settings. */
+		/** Loudness analyzer settings. */
 		TStrongObjectPtr<ULKFSSettings> LKFSSettings;
 
-		/** Handle for results delegate for analyzer. */
+		/** Handle for results delegate for loudness analyzer. */
 		FDelegateHandle LatestOverallLKFSResultsDelegateHandle;
 
-		/** Most recent results from the analyzer. */
+		/** Most recent results from the loudness analyzer. */
 		TOptional<FLKFSResults> LatestLoudnessResults;
 
+		/** True peak analyzer object. */
+		TStrongObjectPtr<UMeterAnalyzer> TruePeakAnalyzer;
+
+		/** True peak analyzer settings. */
+		TStrongObjectPtr<UMeterSettings> TruePeakSettings;
+
+		/** Handle for results delegate for true peak analyzer. */
+		FDelegateHandle LatestOverallTruePeakMeterResultsDelegateHandle;
+
+		/** Most recent results from the true peak analyzer. */
+		TOptional<FMeterResults> LatestTruePeakResults;
+
 		Audio::FDeviceId AudioDeviceId = INDEX_NONE;
 
 		FTSTicker::FDelegateHandle WidgetRefreshTicker;
diff --git a/Engine/Plugins/Runtime/AudioWidgets/Source/AudioWidgets/Private/LoudnessMeterWidgetView.cpp b/Engine/Plugins/Runtime/AudioWidgets/Source/AudioWidgets/Private/LoudnessMeterWidgetView.cpp
--- a/Engine/Plugins/Runtime/AudioWidgets/Source/AudioWidgets/Private/LoudnessMeterWidgetView.cpp
+++ b/Engine/Plugins/Runtime/AudioWidgets/Source/AudioWidgets/Private/LoudnessMeterWidgetView.cpp
@@ -31,6 +31,12 @@ namespace AudioWidgets
 		using FLoudnessScaleParams = FLoudnessMeterWidgetView::FLoudnessScaleParams;
 		using FLoudnessMetricRef = TSharedRef<const FLoudnessMeterWidgetView::FLoudnessMetric>;
 
+		const FNumberFormattingOptions NumberFormattingOptions = FNumberFormattingOptions()
+			.SetMinimumFractionalDigits(1)
+			.SetMaximumFractionalDigits(1);
+
+		const FText NegativeInfinityText = LOCTEXT("NegativeInfinity", "-∞");
+
 		/**
 		 * Derived meter widget type adds bindable attributes for scale DisplayRange and Offset.
 		 */
@@ -46,8 +52,8 @@ namespace AudioWidgets
 
 				SLATE_ATTRIBUTE(FSlateColor, MeterValueColor)
 				SLATE_ATTRIBUTE(TArray<FAudioMeterChannelInfo>, MeterChannelInfo)
-				SLATE_ATTRIBUTE(int, DisplayRange)
-				SLATE_ATTRIBUTE(int, Offset)
+				SLATE_ATTRIBUTE(int32, DisplayRange)
+				SLATE_ATTRIBUTE(int32, Offset)
 
 			SLATE_END_ARGS()
 
@@ -87,8 +93,8 @@ namespace AudioWidgets
 
 		private:
 			const TUniquePtr<FAudioMeterWidgetStyle> AudioMeterWidgetStyle = MakeUnique<FAudioMeterWidgetStyle>(FAudioMeterWidgetStyle::GetDefault());
-			TAttribute<int> Offset;
-			TAttribute<int> DisplayRange;
+			TAttribute<int32> Offset;
+			TAttribute<int32> DisplayRange;
 		};
 
 		/**
@@ -124,13 +130,13 @@ namespace AudioWidgets
 				{
 					return FText::GetEmpty();
 				}
-				else if (*Value == -std::numeric_limits<float>::infinity())
+				else if (*Value <= MIN_VOLUME_DECIBELS)
 				{
-					return LOCTEXT("NegativeInfinity", "-∞");
+					return NegativeInfinityText;
 				}
 
 				const float ConvertedValue = ConvertValue(*Value);
-				return FText::AsNumber(ConvertedValue, GetNumberFormattingOptions());
+				return FText::AsNumber(ConvertedValue, &NumberFormattingOptions);
 			}
 
 			FText FormatRangeText(TOptional<FFloatInterval> Range) const
@@ -141,21 +147,12 @@ namespace AudioWidgets
 				}
 
 				// Range size is already relative units, so no need to call ConvertValue.
-				return FText::AsNumber(Range->Size(), GetNumberFormattingOptions());
+				return FText::AsNumber(Range->Size(), &NumberFormattingOptions);
 			}
 
 		private:
-			static const FNumberFormattingOptions* GetNumberFormattingOptions()
-			{
-				static FNumberFormattingOptions NumberFormattingOptions = FNumberFormattingOptions()
-					.SetMinimumFractionalDigits(1)
-					.SetMaximumFractionalDigits(1);
-
-				return &NumberFormattingOptions;
-			}		
-
 			TAttribute<bool> bRelativeScale;
-			TAttribute<int> Target;
+			TAttribute<int32> Target;
 		};
 
 		/**
@@ -178,14 +175,35 @@ namespace AudioWidgets
 					SNew(STextBlock)
 						.Text_Lambda([InItem, LoudnessValueConverter]()
 							{
-								if (InItem->IsRanged())
-								{
-									return LoudnessValueConverter.FormatRangeText(InItem->Range.Get());
-								}
-								else
+								switch (InItem->MeterMetric.Get())
 								{
-									return LoudnessValueConverter.FormatNumericValueText(InItem->Value.Get());
+									case EAudioMeterMetric::Decibels:
+									{
+										const TOptional<float> ValueOpt = InItem->Value.Get();
+										if (!ValueOpt.IsSet())
+										{
+											return FText::GetEmpty();
+										}
+
+										const float Value = ValueOpt.GetValue();
+
+										return (Value <= MIN_VOLUME_DECIBELS)
+											? NegativeInfinityText
+											: FText::AsNumber(Value, &NumberFormattingOptions);
+									}
+
+									case EAudioMeterMetric::Loudness:
+									{
+										return InItem->IsRanged()
+											? LoudnessValueConverter.FormatRangeText(InItem->Range.Get())
+											: LoudnessValueConverter.FormatNumericValueText(InItem->Value.Get());
+									}
+
+									default:
+										break;
 								}
+
+								return FText::GetEmpty();
 							})
 						.Font(FStyleDefaults::GetFontInfo(14))
 						.ColorAndOpacity_Lambda(GetValueColorAndOpacity)
@@ -195,7 +213,22 @@ namespace AudioWidgets
 				.VAlign(EVerticalAlignment::VAlign_Bottom)
 				[
 					SNew(STextBlock)
-						.Text_Lambda([InItem, bRelativeScale = LoudnessScaleParams.bRelativeScale]() { return (InItem->IsRanged() || bRelativeScale.Get()) ? INVTEXT("LU") : INVTEXT("LUFS"); })
+						.Text_Lambda([InItem, bRelativeScale = LoudnessScaleParams.bRelativeScale]()
+						{
+							switch (InItem->MeterMetric.Get())
+							{
+								case EAudioMeterMetric::Decibels:
+									return INVTEXT("dB");
+
+								case EAudioMeterMetric::Loudness:
+									return (InItem->IsRanged() || bRelativeScale.Get()) ? INVTEXT("LU") : INVTEXT("LUFS");
+
+								default:
+									break;
+							}
+
+							return FText::GetEmpty();
+						})
 						.Font(FStyleDefaults::GetFontInfo(6))
 						.ColorAndOpacity_Lambda(GetValueColorAndOpacity)
 						.Justification(ETextJustify::Right)
@@ -257,9 +290,20 @@ namespace AudioWidgets
 						.MeterChannelInfo_Lambda([InItem, LoudnessValueConverter]() -> TArray<FAudioMeterChannelInfo>
 							{
 								constexpr float DefaultInvalidValue = MIN_VOLUME_DECIBELS;
-								if (InItem->IsRanged())
+
+								if (InItem->MeterMetric.Get() == EAudioMeterMetric::Decibels)
 								{
-									// Display meter bar as a range:
+									return
+									{
+										{
+											.MeterValue = InItem->Value.Get().Get(DefaultInvalidValue),
+											.PeakValue = DefaultInvalidValue
+										}
+									};
+								}
+								else if (InItem->IsRanged())
+								{
+									// Display loudness meter bar as a range:
 									const TOptional<FFloatInterval> Range = InItem->Range.Get();
 									const TOptional<float> ConvertedMin = LoudnessValueConverter.ConvertValue(Range.IsSet() ? Range->Min : TOptional<float>());
 									const TOptional<float> ConvertedMax = LoudnessValueConverter.ConvertValue(Range.IsSet() ? Range->Max : TOptional<float>());
@@ -272,21 +316,27 @@ namespace AudioWidgets
 										}
 									};
 								}
-								else
+
+								// Display standard loudness meter bar:
+								const TOptional<float> ConvertedValue = LoudnessValueConverter.ConvertValue(InItem->Value.Get());
+								return
 								{
-									// Display standard meter bar:
-									const TOptional<float> ConvertedValue = LoudnessValueConverter.ConvertValue(InItem->Value.Get());
-									return
 									{
-										{
-											.MeterValue = ConvertedValue.Get(DefaultInvalidValue),
-											.PeakValue = DefaultInvalidValue
-										}
-									};
-								}
+										.MeterValue = ConvertedValue.Get(DefaultInvalidValue),
+										.PeakValue = DefaultInvalidValue
+									}
+								};
 							})
 						.DisplayRange(LoudnessScaleParams.Range)
-						.Offset_Lambda([Offset = LoudnessScaleParams.Offset, LoudnessValueConverter] { return LoudnessValueConverter.ConvertValue(Offset.Get()); })
+						.Offset_Lambda([InItem, Offset = LoudnessScaleParams.Offset, LoudnessValueConverter]
+						{
+							if (InItem->MeterMetric.Get() == EAudioMeterMetric::Decibels)
+							{
+								return Offset.Get();
+							}
+
+							return static_cast<int32>(LoudnessValueConverter.ConvertValue(Offset.Get()));
+						})
 				]
 				+ SVerticalBox::Slot()
 				.AutoHeight()
@@ -611,7 +661,7 @@ namespace AudioWidgets
 
 					if (LoudnessScaleParams.OnRangeValueCommitted.IsBound())
 					{
-						TSharedRef RangeWidget = SNew(SNumericEntryBox<int>)
+						TSharedRef RangeWidget = SNew(SNumericEntryBox<int32>)
 							.MinDesiredValueWidth(100.0f)
 							.MinSliderValue(10)
 							.MaxSliderValue(70)
@@ -625,7 +675,7 @@ namespace AudioWidgets
 
 					if (LoudnessScaleParams.OnOffsetValueCommitted.IsBound())
 					{
-						TSharedRef OffsetWidget = SNew(SNumericEntryBox<int>)
+						TSharedRef OffsetWidget = SNew(SNumericEntryBox<int32>)
 							.MinDesiredValueWidth(100.0f)
 							.MinSliderValue(-10)
 							.MaxSliderValue(30)
@@ -639,7 +689,7 @@ namespace AudioWidgets
 
 					if (LoudnessScaleParams.OnTargetValueCommitted.IsBound())
 					{
-						TSharedRef TargetWidget = SNew(SNumericEntryBox<int>)
+						TSharedRef TargetWidget = SNew(SNumericEntryBox<int32>)
 							.MinDesiredValueWidth(100.0f)
 							.MinSliderValue(-60)
 							.MaxSliderValue(30)
@@ -786,7 +836,11 @@ namespace AudioWidgets
 			.ItemWidth(MeterWidth)
 			.ItemAlignment(EListItemAlignment::CenterAligned)
 			.SelectionMode(ESelectionMode::Single)
-			.ScrollbarVisibility(EVisibility::Collapsed);
+			.ScrollbarVisibility(EVisibility::Collapsed)
+			.OnMouseButtonClick_Lambda([](FLoudnessMetricRef InItem)
+			{
+				InItem->OnClicked.ExecuteIfBound();
+			});
 
 		MeterTileView->SetItemHeight(TAttribute<float>::CreateLambda([&TileView = *MeterTileView]()
 			{
@@ -892,6 +946,10 @@ namespace AudioWidgets
 							.ItemWidth(60.0f)
 							.ItemAlignment(EListItemAlignment::CenterAligned)
 							.SelectionMode(ESelectionMode::Single)
+							.OnMouseButtonClick_Lambda([](FLoudnessMetricRef InItem)
+							{
+								InItem->OnClicked.ExecuteIfBound();
+							})
 					]
 					+ SVerticalBox::Slot()
 					.AutoHeight()
diff --git a/Engine/Plugins/Runtime/AudioWidgets/Source/AudioWidgets/Private/LoudnessMeterWidgetView.h b/Engine/Plugins/Runtime/AudioWidgets/Source/AudioWidgets/Private/LoudnessMeterWidgetView.h
--- a/Engine/Plugins/Runtime/AudioWidgets/Source/AudioWidgets/Private/LoudnessMeterWidgetView.h
+++ b/Engine/Plugins/Runtime/AudioWidgets/Source/AudioWidgets/Private/LoudnessMeterWidgetView.h
@@ -9,6 +9,12 @@
 
 namespace AudioWidgets
 {
+	enum class EAudioMeterMetric : uint8
+	{
+		Loudness,
+		Decibels
+	};
+
 	/**
 	 * A widget that displays various loudness metrics as numeric values and/or meters.
 	 * Could be moved to AudioWidgetsCore Public once API is stable.
@@ -22,6 +28,7 @@ namespace AudioWidgets
 		{
 			FName Name;
 			FText DisplayName;
+			TAttribute<EAudioMeterMetric> MeterMetric;
 			TAttribute<TOptional<float>> Value;
 			TAttribute<TOptional<FFloatInterval>> Range;
 			TAttribute<bool> bShowValue = false;
@@ -32,6 +39,7 @@ namespace AudioWidgets
 			FOnLinearColorValueChanged OnColorChanged;
 			FOnWindowClosed OnColorPickerWindowClosed;
 			FSimpleDelegate OnResetColorRequested;
+			FSimpleDelegate OnClicked;
 
 			bool IsRanged() const { return Range.IsBound() || Range.Get(NullOpt).IsSet(); }
 		};
@@ -44,20 +52,20 @@ namespace AudioWidgets
 			FSimpleDelegate OnVisibilityToggleRequested;
 		};
 
-		DECLARE_DELEGATE_OneParam(FOnIntValueChanged, int /*NewValue*/);
-		DECLARE_DELEGATE_TwoParams(FOnIntValueCommitted, int /*NewValue*/, ETextCommit::Type /*CommitType*/);
+		DECLARE_DELEGATE_OneParam(FOnIntValueChanged, int32 /*NewValue*/);
+		DECLARE_DELEGATE_TwoParams(FOnIntValueCommitted, int32 /*NewValue*/, ETextCommit::Type /*CommitType*/);
 
 		struct FLoudnessScaleParams
 		{
-			TAttribute<int> Range = 60;
+			TAttribute<int32> Range = 60;
 			FOnIntValueChanged OnRangeValueChanged;
 			FOnIntValueCommitted OnRangeValueCommitted;
 
-			TAttribute<int> Offset = 0;
+			TAttribute<int32> Offset = 0;
 			FOnIntValueChanged OnOffsetValueChanged;
 			FOnIntValueCommitted OnOffsetValueCommitted;
 
-			TAttribute<int> Target = -23;
+			TAttribute<int32> Target = -23;
 			FOnIntValueChanged OnTargetValueChanged;
 			FOnIntValueCommitted OnTargetValueCommitted;
 
diff --git a/Engine/Plugins/Runtime/AudioWidgets/Source/AudioWidgets/Public/LoudnessMeterRackUnitSettings.h b/Engine/Plugins/Runtime/AudioWidgets/Source/AudioWidgets/Public/LoudnessMeterRackUnitSettings.h
--- a/Engine/Plugins/Runtime/AudioWidgets/Source/AudioWidgets/Public/LoudnessMeterRackUnitSettings.h
+++ b/Engine/Plugins/Runtime/AudioWidgets/Source/AudioWidgets/Public/LoudnessMeterRackUnitSettings.h
@@ -25,13 +25,13 @@ struct FLoudnessMeterRackUnitSettings
 	GENERATED_BODY()
 
 	UPROPERTY(EditAnywhere, config, Category = LoudnessMeter, meta = (ClampMin = "10", ClampMax = "70", UIMin = "10", UIMax = "70"))
-	int LoudnessScaleRange = 60;
+	int32 LoudnessScaleRange = 60;
 
 	UPROPERTY(EditAnywhere, config, Category = LoudnessMeter, meta = (ClampMin = "-10", ClampMax = "30", UIMin = "-10", UIMax = "30"))
-	int LoudnessScaleOffset = 0;
+	int32 LoudnessScaleOffset = 0;
 
 	UPROPERTY(EditAnywhere, config, Category = LoudnessMeter, meta = (ClampMin = "-60", ClampMax = "30", UIMin = "-60", UIMax = "30"))
-	int LoudnessScaleTarget = -23;
+	int32 LoudnessScaleTarget = -23;
 
 	UPROPERTY(EditAnywhere, config, Category = LoudnessMeter)
 	bool bUseRelativeLoudnessScale = false;
@@ -51,9 +51,12 @@ struct FLoudnessMeterRackUnitSettings
 	UPROPERTY(EditAnywhere, config, Category = LoudnessMeter)
 	FLoudnessMetricDisplayOptions LoudnessRange = { .bShowValue = true, .bShowMeter = false };
 
-	UPROPERTY(EditAnywhere, config, Category = LoudnessMeter, meta = (ClampMin = "0", ClampMax = "23", UIMin = "0", UIMax = "23"))
-	uint64 ValuesOrderingPermutation = 0; // Max == 4! - 1
+	UPROPERTY(EditAnywhere, config, Category = LoudnessMeter)
+	FLoudnessMetricDisplayOptions TruePeak = { .bShowValue = true, .bShowMeter = false };
+
+	UPROPERTY(EditAnywhere, config, Category = LoudnessMeter, meta = (ClampMin = "0", ClampMax = "119", UIMin = "0", UIMax = "119"))
+	uint64 ValuesOrderingPermutation = 0; // Max == 5! - 1
 
-	UPROPERTY(EditAnywhere, config, Category = LoudnessMeter, meta = (ClampMin = "0", ClampMax = "23", UIMin = "0", UIMax = "23"))
-	uint64 MetersOrderingPermutation = 0; // Max == 4! - 1
+	UPROPERTY(EditAnywhere, config, Category = LoudnessMeter, meta = (ClampMin = "0", ClampMax = "119", UIMin = "0", UIMax = "119"))
+	uint64 MetersOrderingPermutation = 0; // Max == 5! - 1
 };
