[AudioWidgets] Moved LoudnessMeterWidgetView and some of it dependencies to AudioWidgetsCore module so LoudnessMeterWidgetView can be used in Audio Insights standalone.

#Synced-CL 50648262

diff --git a/Engine/Plugins/Runtime/AudioWidgets/Source/AudioWidgets/Private/AudioLoudnessMeter.cpp b/Engine/Plugins/Runtime/AudioWidgets/Source/AudioWidgets/Private/AudioLoudnessMeter.cpp
--- a/Engine/Plugins/Runtime/AudioWidgets/Source/AudioWidgets/Private/AudioLoudnessMeter.cpp
+++ b/Engine/Plugins/Runtime/AudioWidgets/Source/AudioWidgets/Private/AudioLoudnessMeter.cpp
@@ -238,6 +238,8 @@ namespace AudioWidgets
 
 	void FAudioLoudnessMeter::InitLoudnessMeterWidgetView(const FAudioAnalyzerRackUnitConstructParams& Params)
 	{
+		using namespace AudioWidgetsCore;
+
 		// Initializing these callbacks outside of the FAudioLoudnessMeter constructor allows us to use the CreateSP(...) smart pointer variants for attributes/delegates:
 
 		FLoudnessMeterWidgetView::FTimerPanelParams TimerPanelParams
diff --git a/Engine/Plugins/Runtime/AudioWidgets/Source/AudioWidgets/Private/AudioLoudnessMeter.h b/Engine/Plugins/Runtime/AudioWidgets/Source/AudioWidgets/Private/AudioLoudnessMeter.h
--- a/Engine/Plugins/Runtime/AudioWidgets/Source/AudioWidgets/Private/AudioLoudnessMeter.h
+++ b/Engine/Plugins/Runtime/AudioWidgets/Source/AudioWidgets/Private/AudioLoudnessMeter.h
@@ -59,7 +59,7 @@ namespace AudioWidgets
 		FReply HandleResetButtonClicked();
 
 		/** Creates the meters widget */
-		FLoudnessMeterWidgetView LoudnessMeterWidgetView;
+		AudioWidgetsCore::FLoudnessMeterWidgetView LoudnessMeterWidgetView;
 		
 		/** Loudness analyzer object. */
 		TStrongObjectPtr<ULKFSAnalyzer> LKFSAnalyzer;
diff --git a/Engine/Plugins/Runtime/AudioWidgets/Source/AudioWidgets/Private/AudioWidgetsStyle.cpp b/Engine/Plugins/Runtime/AudioWidgets/Source/AudioWidgets/Private/AudioWidgetsStyle.cpp
--- a/Engine/Plugins/Runtime/AudioWidgets/Source/AudioWidgets/Private/AudioWidgetsStyle.cpp
+++ b/Engine/Plugins/Runtime/AudioWidgets/Source/AudioWidgets/Private/AudioWidgetsStyle.cpp
@@ -195,8 +195,6 @@ void FAudioWidgetsStyle::SetupStyles()
 	const FVector2f Icon16x16 = FVector2f(16.0f, 16.0f);
 	const FVector2f Icon20x20 = FVector2f(20.0f, 20.0f);
 
-	Set("AudioWidgetsStyle.ColorMixing", new IMAGE_BRUSH_SVG(TEXT("Slate/Icons/ColorMixing"), Icon16x16));
-
 	Set("AudioWidgetsStyle.FadeIn", new IMAGE_BRUSH_SVG(TEXT("Slate/Icons/FadeIn_20"), Icon20x20));
 	Set("AudioWidgetsStyle.FadeInLinear", new IMAGE_BRUSH_SVG(TEXT("Slate/Icons/FadeInLinear_16"), Icon16x16));
 	Set("AudioWidgetsStyle.FadeInLogarithmic", new IMAGE_BRUSH_SVG(TEXT("Slate/Icons/FadeInLogarithmic_16"), Icon16x16));
@@ -211,8 +209,6 @@ void FAudioWidgetsStyle::SetupStyles()
 
 	Set("AudioWidgetsStyle.MarkerAdd", new IMAGE_BRUSH_SVG(TEXT("Slate/Icons/Marker_20"), Icon20x20));
 	Set("AudioWidgetsStyle.MarkerDelete", new IMAGE_BRUSH_SVG(TEXT("Slate/Icons/MarkerDelete_20"), Icon20x20));
-
-	Set("AudioWidgetsStyle.Reset", new CORE_IMAGE_BRUSH_SVG("Starship/Common/Reset", Icon16x16));
 }
 
 FAudioWidgetsStyle::~FAudioWidgetsStyle()
diff --git a/Engine/Plugins/Runtime/AudioWidgets/Source/AudioWidgets/Private/LoudnessMeterWidgetView.cpp b/Engine/Plugins/Runtime/AudioWidgets/Source/AudioWidgets/Private/LoudnessMeterWidgetView.cpp
deleted file mode 100644
diff --git a/Engine/Plugins/Runtime/AudioWidgets/Source/AudioWidgets/Private/LoudnessMeterWidgetView.h b/Engine/Plugins/Runtime/AudioWidgets/Source/AudioWidgets/Private/LoudnessMeterWidgetView.h
deleted file mode 100644
diff --git a/Engine/Plugins/Runtime/AudioWidgets/Source/AudioWidgets/Private/Permutations.cpp b/Engine/Plugins/Runtime/AudioWidgets/Source/AudioWidgets/Private/Permutations.cpp
deleted file mode 100644
diff --git a/Engine/Plugins/Runtime/AudioWidgets/Source/AudioWidgets/Private/Permutations.h b/Engine/Plugins/Runtime/AudioWidgets/Source/AudioWidgets/Private/Permutations.h
deleted file mode 100644
diff --git a/Engine/Plugins/Runtime/AudioWidgets/Source/AudioWidgets/Private/SDragReorderableTileView.cpp b/Engine/Plugins/Runtime/AudioWidgets/Source/AudioWidgets/Private/SDragReorderableTileView.cpp
deleted file mode 100644
diff --git a/Engine/Plugins/Runtime/AudioWidgets/Source/AudioWidgets/Private/SDragReorderableTileView.h b/Engine/Plugins/Runtime/AudioWidgets/Source/AudioWidgets/Private/SDragReorderableTileView.h
deleted file mode 100644
diff --git a/Engine/Plugins/Runtime/AudioWidgets/Source/AudioWidgets/Public/LoudnessMeterRackUnitSettings.h b/Engine/Plugins/Runtime/AudioWidgets/Source/AudioWidgets/Public/LoudnessMeterRackUnitSettings.h
--- a/Engine/Plugins/Runtime/AudioWidgets/Source/AudioWidgets/Public/LoudnessMeterRackUnitSettings.h
+++ b/Engine/Plugins/Runtime/AudioWidgets/Source/AudioWidgets/Public/LoudnessMeterRackUnitSettings.h
@@ -1,62 +1,19 @@
 // Copyright Epic Games, Inc. All Rights Reserved.
 #pragma once
 
+#include "LoudnessMeterSettings.h"
 #include "Styling/StyleColors.h"
+
 #include "LoudnessMeterRackUnitSettings.generated.h"
 
 USTRUCT()
-struct FLoudnessMetricDisplayOptions
+struct FLoudnessMetricDisplayOptions : public FLoudnessMeterDisplayOptions
 {
 	GENERATED_BODY()
-
-	UPROPERTY(EditAnywhere, config, Category = LoudnessMeter)
-	bool bShowValue = false;
-
-	UPROPERTY(EditAnywhere, config, Category = LoudnessMeter)
-	bool bShowMeter = false;
-
-	UPROPERTY(EditAnywhere, config, Category = LoudnessMeter)
-	FLinearColor Color = FStyleColors::AccentGreen.GetSpecifiedColor().CopyWithNewOpacity(0.5f);
 };
 
 USTRUCT()
-struct FLoudnessMeterRackUnitSettings
+struct FLoudnessMeterRackUnitSettings : public FLoudnessMeterSettings
 {
 	GENERATED_BODY()
-
-	UPROPERTY(EditAnywhere, config, Category = LoudnessMeter, meta = (ClampMin = "10", ClampMax = "70", UIMin = "10", UIMax = "70"))
-	int32 LoudnessScaleRange = 60;
-
-	UPROPERTY(EditAnywhere, config, Category = LoudnessMeter, meta = (ClampMin = "-10", ClampMax = "30", UIMin = "-10", UIMax = "30"))
-	int32 LoudnessScaleOffset = 0;
-
-	UPROPERTY(EditAnywhere, config, Category = LoudnessMeter, meta = (ClampMin = "-60", ClampMax = "30", UIMin = "-60", UIMax = "30"))
-	int32 LoudnessScaleTarget = -23;
-
-	UPROPERTY(EditAnywhere, config, Category = LoudnessMeter)
-	bool bUseRelativeLoudnessScale = false;
-
-	UPROPERTY(EditAnywhere, config, Category = LoudnessMeter)
-	bool bDisplayAnalysisTimer = true;
-
-	UPROPERTY(EditAnywhere, config, Category = LoudnessMeter)
-	FLoudnessMetricDisplayOptions LongTermLoudness = { .bShowValue = true, .bShowMeter = false };
-
-	UPROPERTY(EditAnywhere, config, Category = LoudnessMeter)
-	FLoudnessMetricDisplayOptions ShortTermLoudness = { .bShowValue = false, .bShowMeter = true };
-
-	UPROPERTY(EditAnywhere, config, Category = LoudnessMeter)
-	FLoudnessMetricDisplayOptions MomentaryLoudness = { .bShowValue = false, .bShowMeter = true };
-
-	UPROPERTY(EditAnywhere, config, Category = LoudnessMeter)
-	FLoudnessMetricDisplayOptions LoudnessRange = { .bShowValue = true, .bShowMeter = false };
-
-	UPROPERTY(EditAnywhere, config, Category = LoudnessMeter)
-	FLoudnessMetricDisplayOptions TruePeak = { .bShowValue = true, .bShowMeter = false };
-
-	UPROPERTY(EditAnywhere, config, Category = LoudnessMeter, meta = (ClampMin = "0", ClampMax = "119", UIMin = "0", UIMax = "119"))
-	uint64 ValuesOrderingPermutation = 0; // Max == 5! - 1
-
-	UPROPERTY(EditAnywhere, config, Category = LoudnessMeter, meta = (ClampMin = "0", ClampMax = "119", UIMin = "0", UIMax = "119"))
-	uint64 MetersOrderingPermutation = 0; // Max == 5! - 1
 };
diff --git a/Engine/Plugins/Runtime/AudioWidgets/Source/AudioWidgetsCore/AudioWidgetsCore.Build.cs b/Engine/Plugins/Runtime/AudioWidgets/Source/AudioWidgetsCore/AudioWidgetsCore.Build.cs
--- a/Engine/Plugins/Runtime/AudioWidgets/Source/AudioWidgetsCore/AudioWidgetsCore.Build.cs
+++ b/Engine/Plugins/Runtime/AudioWidgets/Source/AudioWidgetsCore/AudioWidgetsCore.Build.cs
@@ -20,7 +20,10 @@ public class AudioWidgetsCore : ModuleRules
 		(
 			new string[]
 			{
+				"AppFramework",
+				"AudioMixerCore",
 				"CoreUObject",
+				"InputCore",
 				"Slate",
 				"SlateCore",
 			}
diff --git a/Engine/Plugins/Runtime/AudioWidgets/Source/AudioWidgetsCore/Private/AudioWidgetsCoreStyle.cpp b/Engine/Plugins/Runtime/AudioWidgets/Source/AudioWidgetsCore/Private/AudioWidgetsCoreStyle.cpp
new file mode 100644
--- /dev/null
+++ b/Engine/Plugins/Runtime/AudioWidgets/Source/AudioWidgetsCore/Private/AudioWidgetsCoreStyle.cpp
@@ -0,0 +1,34 @@
+// Copyright Epic Games, Inc. All Rights Reserved.
+
+#include "AudioWidgetsCoreStyle.h"
+
+#include "Styling/AppStyle.h"
+#include "Styling/SlateStyleMacros.h"
+#include "Styling/SlateStyleRegistry.h"
+
+FAudioWidgetsCoreStyle::FAudioWidgetsCoreStyle()
+	: FSlateStyleSet(StyleName)
+{
+	SetParentStyleName(FAppStyle::GetAppStyleSetName());
+	SetContentRoot(FPaths::EnginePluginsDir() / TEXT("Runtime/AudioWidgets/Content"));
+	SetCoreContentRoot(FPaths::EngineContentDir());
+
+	const FVector2f Icon16x16 = FVector2f(16.0f, 16.0f);
+
+	Set("AudioWidgetsCoreStyle.ColorMixing", new IMAGE_BRUSH_SVG(TEXT("Slate/Icons/ColorMixing"), Icon16x16));
+	Set("AudioWidgetsCoreStyle.Reset", new CORE_IMAGE_BRUSH_SVG("Editor/Slate/Starship/Common/Reset", Icon16x16));
+	Set("AudioWidgetsCoreStyle.Settings", new CORE_IMAGE_BRUSH_SVG("Slate/Starship/Common/Settings", Icon16x16));
+
+	FSlateStyleRegistry::RegisterSlateStyle(*this);
+}
+
+FAudioWidgetsCoreStyle::~FAudioWidgetsCoreStyle()
+{
+	FSlateStyleRegistry::UnRegisterSlateStyle(*this);
+}
+
+const FAudioWidgetsCoreStyle& FAudioWidgetsCoreStyle::Get()
+{
+	static FAudioWidgetsCoreStyle Instance;
+	return Instance;
+}
diff --git a/Engine/Plugins/Runtime/AudioWidgets/Source/AudioWidgetsCore/Private/LoudnessMeterWidgetView.cpp b/Engine/Plugins/Runtime/AudioWidgets/Source/AudioWidgetsCore/Private/LoudnessMeterWidgetView.cpp
new file mode 100644
--- /dev/null
+++ b/Engine/Plugins/Runtime/AudioWidgets/Source/AudioWidgetsCore/Private/LoudnessMeterWidgetView.cpp
@@ -0,0 +1,1101 @@
+// Copyright Epic Games, Inc. All Rights Reserved.
+
+#include "LoudnessMeterWidgetView.h"
+
+#include "AudioDefines.h"
+#include "AudioWidgetsCoreStyle.h"
+#include "AudioMeterWidgetStyle.h"
+#include "Brushes/SlateColorBrush.h"
+#include "Brushes/SlateRoundedBoxBrush.h"
+#include "Framework/Application/SlateApplication.h"
+#include "Framework/MultiBox/MultiBoxBuilder.h"
+#include "Permutations.h"
+#include "SAudioMeterWidget.h"
+#include "SDragReorderableTileView.h"
+#include "Styling/StyleColors.h"
+#include "Widgets/Colors/SColorBlock.h"
+#include "Widgets/Colors/SColorPicker.h"
+#include "Widgets/Input/SButton.h"
+#include "Widgets/Input/SNumericEntryBox.h"
+#include "Widgets/Layout/SSeparator.h"
+#include "Widgets/Layout/SSpacer.h"
+
+#include <limits>
+
+#define LOCTEXT_NAMESPACE "FLoudnessMeterWidgetView"
+
+namespace AudioWidgetsCore
+{
+	namespace LoudnessMeterWidgetView_Private
+	{
+		using FLoudnessScaleParams = FLoudnessMeterWidgetView::FLoudnessScaleParams;
+		using FLoudnessMetricRef = TSharedRef<const FLoudnessMeterWidgetView::FLoudnessMetric>;
+
+		const FNumberFormattingOptions NumberFormattingOptions = FNumberFormattingOptions()
+			.SetMinimumFractionalDigits(1)
+			.SetMaximumFractionalDigits(1);
+
+		const FText NegativeInfinityText = LOCTEXT("NegativeInfinity", "-∞");
+
+		/**
+		 * Derived meter widget type adds bindable attributes for scale DisplayRange and Offset.
+		 */
+		class SVariableDisplayRangeMeterWidget : public SAudioMeterWidget
+		{
+		public:
+			SLATE_BEGIN_ARGS(SVariableDisplayRangeMeterWidget)
+				: _MeterValueColor(FLinearColor::Green)
+				, _DisplayRange(60)
+				, _Offset(0)
+				{
+				}
+
+				SLATE_ATTRIBUTE(FSlateColor, MeterValueColor)
+				SLATE_ATTRIBUTE(TArray<FAudioMeterChannelInfo>, MeterChannelInfo)
+				SLATE_ATTRIBUTE(int32, DisplayRange)
+				SLATE_ATTRIBUTE(int32, Offset)
+
+			SLATE_END_ARGS()
+
+			void Construct(const SVariableDisplayRangeMeterWidget::FArguments& InArgs)
+			{
+				AudioMeterWidgetStyle->SetMeterSize({ 700.0, 25.0 }); // MaxLength, Width.
+				AudioMeterWidgetStyle->SetScaleHashHeight(4.0f); // This is the tick mark length.
+
+				DisplayRange = InArgs._DisplayRange;
+				Offset = InArgs._Offset;
+
+				SAudioMeterWidget::Construct(
+					SAudioMeterWidget::FArguments()
+						.Orientation(EOrientation::Orient_Vertical)
+						.BackgroundColor(FLinearColor::Transparent)
+						.MeterBackgroundColor(FAudioMeterDefaultColorWidgetStyle::GetDefault().MeterBackgroundColor)
+						.MeterPeakColor(FLinearColor::Transparent)
+						.MeterClippingColor(FLinearColor::Transparent)
+						.MeterScaleColor(FLinearColor::White.CopyWithNewOpacity(0.25f))
+						.MeterScaleLabelColor(FLinearColor::White.CopyWithNewOpacity(0.25f))
+						.MeterValueColor(InArgs._MeterValueColor)
+						.MeterChannelInfo(InArgs._MeterChannelInfo)
+						.Style(AudioMeterWidgetStyle.Get())
+					);
+			}
+
+			virtual int32 OnPaint(const FPaintArgs& Args, const FGeometry& AllottedGeometry, const FSlateRect& MyCullingRect, FSlateWindowElementList& OutDrawElements, int32 LayerId, const FWidgetStyle& InWidgetStyle, bool bParentEnabled) const override
+			{
+				// Update the style from our attribute bindings:
+				const float DisplayRangeValue = DisplayRange.Get();
+				const float OffsetValue = Offset.Get();
+				AudioMeterWidgetStyle->SetValueRangeDb({ OffsetValue - DisplayRangeValue, OffsetValue });
+
+				// Paint the widget using the updated style:
+				return SAudioMeterWidget::OnPaint(Args, AllottedGeometry, MyCullingRect, OutDrawElements, LayerId, InWidgetStyle, bParentEnabled);
+			}
+
+		private:
+			const TUniquePtr<FAudioMeterWidgetStyle> AudioMeterWidgetStyle = MakeUnique<FAudioMeterWidgetStyle>(FAudioMeterWidgetStyle::GetDefault());
+			TAttribute<int32> Offset;
+			TAttribute<int32> DisplayRange;
+		};
+
+		/**
+		 * If the Relative Scale attribute is true this will convert values given in LUFS into values in LU, relative to the Target attribute.
+		 * Otherwise it will leave the LUFS values as-is.
+		 */
+		class FLoudnessValueConverter
+		{
+		public:
+			FLoudnessValueConverter(const FLoudnessScaleParams& LoudnessScaleParams)
+				: bRelativeScale(LoudnessScaleParams.bRelativeScale)
+				, Target(LoudnessScaleParams.Target)
+			{
+			}
+
+			float ConvertValue(float Value) const
+			{
+				return (bRelativeScale.Get()) ? (Value - Target.Get()) : Value;
+			}
+
+			TOptional<float> ConvertValue(TOptional<float> Value) const
+			{
+				if (Value.IsSet())
+				{
+					return ConvertValue(*Value);
+				}
+				return NullOpt;
+			}
+
+			FText FormatNumericValueText(TOptional<float> Value) const
+			{
+				if (!Value.IsSet())
+				{
+					return FText::GetEmpty();
+				}
+				else if (*Value <= MIN_VOLUME_DECIBELS)
+				{
+					return NegativeInfinityText;
+				}
+
+				const float ConvertedValue = ConvertValue(*Value);
+				return FText::AsNumber(ConvertedValue, &NumberFormattingOptions);
+			}
+
+			FText FormatRangeText(TOptional<FFloatInterval> Range) const
+			{
+				if (!Range.IsSet())
+				{
+					return FText::GetEmpty();
+				}
+
+				// Range size is already relative units, so no need to call ConvertValue.
+				return FText::AsNumber(Range->Size(), &NumberFormattingOptions);
+			}
+
+		private:
+			TAttribute<bool> bRelativeScale;
+			TAttribute<int32> Target;
+		};
+
+		/**
+		 * Creates a widget that displays the numeric value of the given loudness metric, with units.
+		 */
+		TSharedRef<SWidget> MakeValueReadoutWidget(FLoudnessMetricRef InItem, const FLoudnessScaleParams& LoudnessScaleParams)
+		{
+			const FLoudnessValueConverter LoudnessValueConverter(LoudnessScaleParams);
+
+			const auto GetValueColorAndOpacity = [InItem]() { return InItem->Color.Get().CopyWithNewOpacity(1.0f); };
+
+			return SNew(SOverlay)
+				+ SOverlay::Slot()
+				[
+					SNew(SSpacer)
+						.Size(FVector2D(46.0f, 28.0f))
+				]
+				+ SOverlay::Slot()
+				[
+					SNew(STextBlock)
+						.Text_Lambda([InItem, LoudnessValueConverter]()
+							{
+								switch (InItem->MeterMetric.Get())
+								{
+									case EAudioMeterMetric::Decibels:
+									{
+										const TOptional<float> ValueOpt = InItem->Value.Get();
+										if (!ValueOpt.IsSet())
+										{
+											return FText::GetEmpty();
+										}
+
+										const float Value = ValueOpt.GetValue();
+
+										return (Value <= MIN_VOLUME_DECIBELS)
+											? NegativeInfinityText
+											: FText::AsNumber(Value, &NumberFormattingOptions);
+									}
+
+									case EAudioMeterMetric::Loudness:
+									{
+										return InItem->IsRanged()
+											? LoudnessValueConverter.FormatRangeText(InItem->Range.Get())
+											: LoudnessValueConverter.FormatNumericValueText(InItem->Value.Get());
+									}
+
+									default:
+										break;
+								}
+
+								return FText::GetEmpty();
+							})
+						.Font(FStyleDefaults::GetFontInfo(14))
+						.ColorAndOpacity_Lambda(GetValueColorAndOpacity)
+						.Justification(ETextJustify::Right)
+				]
+				+ SOverlay::Slot()
+				.VAlign(EVerticalAlignment::VAlign_Bottom)
+				[
+					SNew(STextBlock)
+						.Text_Lambda([InItem, bRelativeScale = LoudnessScaleParams.bRelativeScale]()
+						{
+							switch (InItem->MeterMetric.Get())
+							{
+								case EAudioMeterMetric::Decibels:
+									return INVTEXT("dB");
+
+								case EAudioMeterMetric::Loudness:
+									return (InItem->IsRanged() || bRelativeScale.Get()) ? INVTEXT("LU") : INVTEXT("LUFS");
+
+								default:
+									break;
+							}
+
+							return FText::GetEmpty();
+						})
+						.Font(FStyleDefaults::GetFontInfo(6))
+						.ColorAndOpacity_Lambda(GetValueColorAndOpacity)
+						.Justification(ETextJustify::Right)
+				];
+		}
+
+		/**
+		 * Creates a widget with the actual content of a Value tile, containing the numeric value readout widget and the loudness metric display name.
+		 */
+		TSharedRef<SWidget> MakeValueTileInnerContent(FLoudnessMetricRef InItem, FLoudnessScaleParams LoudnessScaleParams)
+		{
+			return SNew(SVerticalBox)
+				+ SVerticalBox::Slot()
+				.MinHeight(8)
+				.MaxHeight(8)
+				[
+					SNew(SSpacer)
+						.Size(FVector2D(60.0f, 8.0f))
+				]
+				+ SVerticalBox::Slot()
+				.AutoHeight()
+				.HAlign(EHorizontalAlignment::HAlign_Center)
+				[
+					MakeValueReadoutWidget(InItem, LoudnessScaleParams)
+				]
+				+ SVerticalBox::Slot()
+				.MinHeight(24)
+				.MaxHeight(24)
+				.HAlign(EHorizontalAlignment::HAlign_Center)
+				.VAlign(EVerticalAlignment::VAlign_Center)
+				[
+					SNew(STextBlock)
+						.Text(InItem->DisplayName)
+						.Font(FStyleDefaults::GetFontInfo(8))
+						.AutoWrapText(true)
+						.Justification(ETextJustify::Center)
+				];
+		}
+
+		/**
+		 * Creates a widget with the actual content of a Meter tile, containing the numeric value readout widget, the actual meter widget, and the loudness metric display name.
+		 */
+		TSharedRef<SWidget> MakeMeterTileInnerContent(FLoudnessMetricRef InItem, FLoudnessScaleParams LoudnessScaleParams)
+		{
+			const FLoudnessValueConverter LoudnessValueConverter(LoudnessScaleParams);
+
+			return SNew(SVerticalBox)
+				+ SVerticalBox::Slot()
+				.AutoHeight()
+				.HAlign(EHorizontalAlignment::HAlign_Center)
+				[
+					MakeValueReadoutWidget(InItem, LoudnessScaleParams)
+				]
+				+ SVerticalBox::Slot()
+				.HAlign(EHorizontalAlignment::HAlign_Center)
+				[
+					SNew(SVariableDisplayRangeMeterWidget)
+						.MeterValueColor_Lambda([InItem]() { return InItem->Color.Get(); })
+						.MeterChannelInfo_Lambda([InItem, LoudnessValueConverter]() -> TArray<FAudioMeterChannelInfo>
+							{
+								constexpr float DefaultInvalidValue = MIN_VOLUME_DECIBELS;
+
+								if (InItem->MeterMetric.Get() == EAudioMeterMetric::Decibels)
+								{
+									return
+									{
+										{
+											.MeterValue = InItem->Value.Get().Get(DefaultInvalidValue),
+											.PeakValue = DefaultInvalidValue
+										}
+									};
+								}
+								else if (InItem->IsRanged())
+								{
+									// Display loudness meter bar as a range:
+									const TOptional<FFloatInterval> Range = InItem->Range.Get();
+									const TOptional<float> ConvertedMin = LoudnessValueConverter.ConvertValue(Range.IsSet() ? Range->Min : TOptional<float>());
+									const TOptional<float> ConvertedMax = LoudnessValueConverter.ConvertValue(Range.IsSet() ? Range->Max : TOptional<float>());
+									return
+									{
+										{
+											.MeterValue = ConvertedMax.Get(DefaultInvalidValue),
+											.PeakValue = DefaultInvalidValue,
+											.MeterValueBarStart = ConvertedMin
+										}
+									};
+								}
+
+								// Display standard loudness meter bar:
+								const TOptional<float> ConvertedValue = LoudnessValueConverter.ConvertValue(InItem->Value.Get());
+								return
+								{
+									{
+										.MeterValue = ConvertedValue.Get(DefaultInvalidValue),
+										.PeakValue = DefaultInvalidValue
+									}
+								};
+							})
+						.DisplayRange(LoudnessScaleParams.Range)
+						.Offset_Lambda([InItem, Offset = LoudnessScaleParams.Offset, LoudnessValueConverter]
+						{
+							if (InItem->MeterMetric.Get() == EAudioMeterMetric::Decibels)
+							{
+								return Offset.Get();
+							}
+
+							return static_cast<int32>(LoudnessValueConverter.ConvertValue(Offset.Get()));
+						})
+				]
+				+ SVerticalBox::Slot()
+				.AutoHeight()
+				.MinHeight(28)
+				.HAlign(EHorizontalAlignment::HAlign_Center)
+				.VAlign(EVerticalAlignment::VAlign_Center)
+				[
+					SNew(STextBlock)
+						.Text_Lambda([InItem]() { return InItem->DisplayName; })
+						.Font(FStyleDefaults::GetFontInfo(8))
+						.AutoWrapText(true)
+						.Justification(ETextJustify::Center)
+				];
+		}
+
+		DECLARE_DELEGATE_RetVal_OneParam(TSharedRef<SWidget>, FOnMakeTileInnerContent, FLoudnessMetricRef);
+
+		/**
+		 * Represents a loudness metric that is being dragged from a given source tile view.
+		 * A decorator renders the content of the tile, with a rounded border selection highlight.
+		 */
+		class FLoudnessValueDragDropOp : public TItemDragDropOperation<FLoudnessMetricRef, FLoudnessValueDragDropOp>
+		{
+		public:
+			using Base = TItemDragDropOperation<FLoudnessMetricRef, FLoudnessValueDragDropOp>;
+
+			DRAG_DROP_OPERATOR_TYPE(FLoudnessValueDragDropOp, Base)
+
+			FLoudnessValueDragDropOp(TSharedRef<Base::STileViewType> InSourceTileView, FLoudnessMetricRef InLoudnessMetric, FOnMakeTileInnerContent InOnMakeTileInnerContent)
+				: Base(InSourceTileView)
+				, LoudnessMetric(InLoudnessMetric)
+				, OnMakeTileInnerContent(InOnMakeTileInnerContent)
+			{
+				Construct();
+			}
+
+			virtual TSharedPtr<SWidget> GetDefaultDecorator() const override;
+
+			virtual const FLoudnessMetricRef& GetItem() const final { return LoudnessMetric; }
+
+		private:
+			const FLoudnessMetricRef LoudnessMetric;
+			const FOnMakeTileInnerContent OnMakeTileInnerContent;
+		};
+
+		TSharedPtr<SWidget> FLoudnessValueDragDropOp::GetDefaultDecorator() const
+		{
+			static const FSlateColorBrush BackgroundPanelBrush(FStyleColors::Panel);
+			static const FSlateRoundedBoxBrush BorderBrush(FStyleColors::Select, 4.0f);
+			static const FSlateRoundedBoxBrush InteriorPanelBrush(FStyleColors::Panel, 4.0f);
+
+			TSharedRef<SWidget> TileInnerContent = (OnMakeTileInnerContent.IsBound()) ? OnMakeTileInnerContent.Execute(LoudnessMetric) : SNullWidget::NullWidget;
+
+			return SNew(SBox)
+				.MaxDesiredHeight(180.0f)
+				.MinDesiredWidth(66.0f)
+				[
+					SNew(SBorder)
+						.BorderImage(&BackgroundPanelBrush)
+						.Padding(0.0f)
+						[
+							SNew(SBorder)
+								.BorderImage(&BorderBrush)
+								.Padding(1.0f)
+								[
+									SNew(SBorder)
+										.BorderImage(&InteriorPanelBrush)
+										[
+											TileInnerContent
+										]
+								]
+						]
+				];
+		}
+
+		/**
+		 * Delegate adapter to execute the necessary delegate to show or hide the numeric value of the given loudness metric.
+		 */
+		void SetShowLoudnessMetricValue(const FLoudnessMetricRef& LoudnessMetric, bool bShowValue)
+		{
+			if (LoudnessMetric->bShowValue.Get() != bShowValue)
+			{
+				LoudnessMetric->OnShowValueToggleRequested.ExecuteIfBound();
+			}
+		}
+
+		/**
+		 * Delegate adapter to execute the necessary delegate to show or hide the meter for the given loudness metric.
+		 */
+		void SetShowLoudnessMetricMeter(const FLoudnessMetricRef& LoudnessMetric, bool bShowMeter)
+		{
+			if (LoudnessMetric->bShowMeter.Get() != bShowMeter)
+			{
+				LoudnessMetric->OnShowMeterToggleRequested.ExecuteIfBound();
+			}
+		}
+
+		/**
+		 * TileView widget for displaying an observable collection of loudness metrics.
+		 * Dropping of a dragged loudness metric onto the tile view is supported if the tile view is currently empty.
+		 */
+		class SLoudnessMetricTileView : public SDragReorderableTileView<FLoudnessValueDragDropOp>
+		{
+		public:
+			using Base = SDragReorderableTileView<FLoudnessValueDragDropOp>;
+
+			void Construct(const Base::FArguments& InArgs, FOnMakeTileInnerContent InOnMakeTileInnerContent, FOnSetShowItem InOnSetShowItem, FOnReorderItems InOnReorderItems)
+			{
+				Base::Construct(InArgs, InOnSetShowItem, InOnReorderItems);
+
+				OnMakeTileInnerContent = InOnMakeTileInnerContent;
+			}
+
+			TSharedRef<SWidget> MakeTileInnerContent(const FLoudnessMetricRef& LoudnessMetric) const
+			{
+				return (OnMakeTileInnerContent.IsBound()) ? OnMakeTileInnerContent.Execute(LoudnessMetric) : SNullWidget::NullWidget;
+			}
+
+			virtual TSharedPtr<FLoudnessValueDragDropOp> CreateDragDropOperation() override
+			{
+				ensure(GetSelectionMode() == ESelectionMode::Single);
+				if (FLoudnessMetricRef* SelectedItem = SelectedItems.FindArbitraryElement()) // As we are using ESelectionMode::Single an arbitrary element is the only element.
+				{
+					TSharedRef<FLoudnessValueDragDropOp> DragDropOperation = MakeShared<FLoudnessValueDragDropOp>(SharedThis(this), *SelectedItem, OnMakeTileInnerContent);
+
+					// Hold a weak pointer to maintain awareness of whether a tile is currently being dragged:
+					WeakDragDropOperation = DragDropOperation;
+
+					return DragDropOperation;
+				}
+				return nullptr;
+			}
+
+			bool HasDragDropOperation() const { return WeakDragDropOperation.IsValid(); }
+
+		private:
+			FOnMakeTileInnerContent OnMakeTileInnerContent;
+			TWeakPtr<FDragDropOperation> WeakDragDropOperation;
+		};
+
+
+		/**
+		 * Tile widget for a single loudness metric. Dragging and dropping tiles is supported.
+		 */
+		class SLoudnessMetricTile : public SDragReorderableTile<FLoudnessValueDragDropOp>
+		{
+		public:
+			const FSlateBrush* GetHoveredBorder() const
+			{
+				if (IsHovered())
+				{
+					static const FSlateRoundedBoxBrush RoundedHoverBrush(FStyleColors::Hover, 4.0f);
+					return &RoundedHoverBrush;
+				}
+				return nullptr;
+			}
+
+			FLinearColor GetOverlayColor() const
+			{
+				if (IsBeingDragged())
+				{
+					return FStyleColors::Recessed.GetSpecifiedColor().CopyWithNewOpacity(0.9f);
+				}
+				return FLinearColor::Transparent;
+			}
+		};
+
+
+		/**
+		 * Helper function to create an empty tile widget, with common styling applied.
+		 */
+		TSharedRef<SLoudnessMetricTile> MakeEmptyTileWidget(const TSharedRef<STableViewBase>& OwnerTable)
+		{
+			const auto MakeTableRowStyle = []() -> FTableRowStyle
+				{
+					const FSlateRoundedBoxBrush RoundedPanelBrush(FStyleColors::Panel, 4.0f);
+
+					FTableRowStyle TableRowStyle = FAppStyle::Get().GetWidgetStyle<FTableRowStyle>("TableView.Row");
+					TableRowStyle.SetEvenRowBackgroundHoveredBrush(RoundedPanelBrush);
+					TableRowStyle.SetOddRowBackgroundHoveredBrush(RoundedPanelBrush);
+					TableRowStyle.SetSelectorFocusedBrush(FSlateNoResource());
+					return TableRowStyle;
+				};
+			static FTableRowStyle LoudnessMetricTableRowStyle = MakeTableRowStyle();
+
+			return SNew(SLoudnessMetricTile, OwnerTable)
+				.ShowSelection(false)
+				.Style(&LoudnessMetricTableRowStyle);
+		}
+
+		/**
+		 * Callback to create a tile widget for the numeric Values tile view. A color block overlay will be used to darken the tile content when it's being dragged.
+		 */
+		TSharedRef<ITableRow> MakeNumericValueTileWidget(FLoudnessMetricRef InItem, const TSharedRef<STableViewBase>& OwnerTable)
+		{
+			TSharedRef<SLoudnessMetricTileView> OwnerTileView = StaticCastSharedRef<SLoudnessMetricTileView>(OwnerTable);
+
+			TSharedRef<SLoudnessMetricTile> TileWidget = MakeEmptyTileWidget(OwnerTileView);
+
+			TileWidget->SetContent(
+				SNew(SOverlay)
+				+ SOverlay::Slot()
+				[
+					OwnerTileView->MakeTileInnerContent(InItem)
+				]
+				+ SOverlay::Slot()
+				[
+					SNew(SColorBlock)
+						.Color(TileWidget, &SLoudnessMetricTile::GetOverlayColor)
+				]);
+
+			return TileWidget;
+		}
+
+		/**
+		 * Callback to create a tile widget for the Meters tile view. A color block overlay will be used to darken the tile content when it's being dragged.
+		 */
+		TSharedRef<ITableRow> MakeMeterValueTileWidget(FLoudnessMetricRef InItem, const TSharedRef<STableViewBase>& OwnerTable)
+		{
+			static const FSlateRoundedBoxBrush RoundedPanelBrush(FStyleColors::Panel, 4.0f);
+
+			TSharedRef<SLoudnessMetricTileView> OwnerTileView = StaticCastSharedRef<SLoudnessMetricTileView>(OwnerTable);
+
+			TSharedRef<SLoudnessMetricTile> TileWidget = MakeEmptyTileWidget(OwnerTileView);
+
+			TileWidget->SetContent(
+				SNew(SOverlay)
+				+ SOverlay::Slot()
+				[
+					SNew(SColorBlock)
+						.Color(FStyleColors::Recessed.GetSpecifiedColor())
+				]
+				+ SOverlay::Slot()
+				[
+					SNew(SBox)
+						.VAlign(EVerticalAlignment::VAlign_Center)
+						.Padding(3.0f)
+						[
+							SNew(SBorder)
+								.BorderImage(TileWidget, &SLoudnessMetricTile::GetHoveredBorder)
+								.Padding(1.0f)
+								[
+									SNew(SBorder)
+										.BorderImage(&RoundedPanelBrush)
+										.Padding(0.0f, 8.0f)
+										[
+											OwnerTileView->MakeTileInnerContent(InItem)
+										]
+								]
+						]
+				]
+				+ SOverlay::Slot()
+				[
+					SNew(SColorBlock)
+						.Color(TileWidget, &SLoudnessMetricTile::GetOverlayColor)
+				]);
+
+			return TileWidget;
+		}
+	} // namespace LoudnessMeterWidgetView_Private
+
+	FLinearColor FLoudnessMeterWidgetView::GetDefaultMeterColor()
+	{
+		return FStyleColors::AccentGreen.GetSpecifiedColor().CopyWithNewOpacity(0.5f);
+	}
+
+	FLoudnessMeterWidgetView::FLoudnessMeterWidgetView()
+		: LoudnessMetrics(MakeShared<FLoudnessMetricRefArray>())
+		, ValuesCollectionView(MakeShared<FLoudnessMetricCollectionView>(LoudnessMetrics, [](const FLoudnessMetricRef& LoudnessMetricRef) { return LoudnessMetricRef->bShowValue.Get(); }))
+		, MetersCollectionView(MakeShared<FLoudnessMetricCollectionView>(LoudnessMetrics, [](const FLoudnessMetricRef& LoudnessMetricRef) { return LoudnessMetricRef->bShowMeter.Get(); }))
+	{
+	}
+
+	FLoudnessMeterWidgetView::~FLoudnessMeterWidgetView()
+	{
+		LoudnessMetrics->Reset();
+		ValuesCollectionView->ResetVisibleItems();
+		MetersCollectionView->ResetVisibleItems();
+	}
+
+	void FLoudnessMeterWidgetView::InitTimerPanel(const FTimerPanelParams& InTimerPanelParams)
+	{
+		TimerPanelParams = InTimerPanelParams;
+	}
+
+	void FLoudnessMeterWidgetView::InitLoudnessScale(const FLoudnessScaleParams& InLoudnessScaleParams)
+	{
+		LoudnessScaleParams = InLoudnessScaleParams;
+	}
+
+	void FLoudnessMeterWidgetView::InitValuesDisplayOrder(const FDisplayOrderParams& InDisplayOrderParams)
+	{
+		ValuesCollectionView->SetDisplayOrderParams(InDisplayOrderParams);
+	}
+
+	void FLoudnessMeterWidgetView::InitMetersDisplayOrder(const FDisplayOrderParams& InDisplayOrderParams)
+	{
+		MetersCollectionView->SetDisplayOrderParams(InDisplayOrderParams);
+	}
+
+	TSharedRef<SWidget> FLoudnessMeterWidgetView::MakeWidget() const
+	{
+		using namespace LoudnessMeterWidgetView_Private;
+
+		static FSlateColorBrush BackgroundBrush(FStyleColors::Recessed);
+
+		constexpr bool bForceSmallIcons = true;
+		const TSharedRef<FUICommandList> CommandList = MakeShared<FUICommandList>();
+		FSlimHorizontalToolBarBuilder ToolBarBuilder(CommandList, FMultiBoxCustomization::None, nullptr, bForceSmallIcons);
+
+		const auto MakeOptionsMenu = [
+			CommandList,
+			TimerPanelParams = this->TimerPanelParams,
+			LoudnessScaleParams = this->LoudnessScaleParams,
+			LoudnessMetrics = this->LoudnessMetrics]() -> TSharedRef<SWidget>
+			{
+				constexpr bool bShouldCloseWindowAfterMenuSelection = true;
+				FMenuBuilder MenuBuilder(bShouldCloseWindowAfterMenuSelection, CommandList);
+
+				{
+					MenuBuilder.BeginSection(NAME_None, LOCTEXT("LoudnessScale", "Loudness Scale"));
+
+					if (LoudnessScaleParams.OnRangeValueCommitted.IsBound())
+					{
+						TSharedRef RangeWidget = SNew(SNumericEntryBox<int32>)
+							.MinDesiredValueWidth(100.0f)
+							.MinSliderValue(10)
+							.MaxSliderValue(70)
+							.AllowSpin(true)
+							.Value_Lambda([Range = LoudnessScaleParams.Range]() { return Range.Get(); })
+							.OnValueChanged(LoudnessScaleParams.OnRangeValueChanged)
+							.OnValueCommitted(LoudnessScaleParams.OnRangeValueCommitted);
+
+						MenuBuilder.AddWidget(RangeWidget, LOCTEXT("LoudnessScale_Range", "Range"));
+					}
+
+					if (LoudnessScaleParams.OnOffsetValueCommitted.IsBound())
+					{
+						TSharedRef OffsetWidget = SNew(SNumericEntryBox<int32>)
+							.MinDesiredValueWidth(100.0f)
+							.MinSliderValue(-10)
+							.MaxSliderValue(30)
+							.AllowSpin(true)
+							.Value_Lambda([Offset = LoudnessScaleParams.Offset]() { return Offset.Get(); })
+							.OnValueChanged(LoudnessScaleParams.OnOffsetValueChanged)
+							.OnValueCommitted(LoudnessScaleParams.OnOffsetValueCommitted);
+
+						MenuBuilder.AddWidget(OffsetWidget, LOCTEXT("LoudnessScale_Offset", "Offset"));
+					}
+
+					if (LoudnessScaleParams.OnTargetValueCommitted.IsBound())
+					{
+						TSharedRef TargetWidget = SNew(SNumericEntryBox<int32>)
+							.MinDesiredValueWidth(100.0f)
+							.MinSliderValue(-60)
+							.MaxSliderValue(30)
+							.AllowSpin(true)
+							.Value_Lambda([Target = LoudnessScaleParams.Target]() { return Target.Get(); })
+							.OnValueChanged(LoudnessScaleParams.OnTargetValueChanged)
+							.OnValueCommitted(LoudnessScaleParams.OnTargetValueCommitted);
+
+						MenuBuilder.AddWidget(TargetWidget, LOCTEXT("LoudnessScale_Target", "Target"));
+					}
+
+					if (LoudnessScaleParams.OnRelativeScaleToggleRequested.IsBound())
+					{
+						MenuBuilder.AddMenuEntry(
+							LOCTEXT("LoudnessScale_RelativeScale_Label", "Relative Scale"),
+							LOCTEXT("LoudnessScale_RelativeScale_ToolTip", "Loudness values should be displayed relative to the Target loudness"),
+							FSlateIcon(),
+							FUIAction(
+								LoudnessScaleParams.OnRelativeScaleToggleRequested,
+								FCanExecuteAction(),
+								FIsActionChecked::CreateLambda([bRelativeScale = LoudnessScaleParams.bRelativeScale]() { return bRelativeScale.Get(); })
+							),
+							NAME_None,
+							EUserInterfaceActionType::ToggleButton);
+					}
+
+					MenuBuilder.EndSection();
+				}
+
+				{
+					MenuBuilder.BeginSection(NAME_None, LOCTEXT("LayoutOptions", "Layout Options"));
+
+					if (TimerPanelParams.OnVisibilityToggleRequested.IsBound())
+					{
+						MenuBuilder.AddMenuEntry(
+							LOCTEXT("LayoutOptions_Time_Label", "Time"),
+							LOCTEXT("LayoutOptions_Time_ToolTip", "Display analysis timer"),
+							FSlateIcon(),
+							FUIAction(
+								TimerPanelParams.OnVisibilityToggleRequested,
+								FCanExecuteAction(),
+								FIsActionChecked::CreateLambda([bIsVisible = TimerPanelParams.bIsVisible]() { return bIsVisible.Get(); })
+							),
+							NAME_None,
+							EUserInterfaceActionType::ToggleButton);
+					}
+
+					for (const FLoudnessMetricRef& LoudnessMetricRef : LoudnessMetrics.Get())
+					{
+						MenuBuilder.AddSubMenu(
+							LoudnessMetricRef->DisplayName,
+							FText(),
+							FNewMenuDelegate::CreateSPLambda(LoudnessMetricRef, [&LoudnessMetric = LoudnessMetricRef.Get()](FMenuBuilder& SubMenu)
+								{
+									SubMenu.BeginSection(NAME_None, LoudnessMetric.DisplayName);
+
+									SubMenu.AddMenuEntry(
+										LOCTEXT("LayoutOptions_LoudnessMetric_ShowValue", "Show Value"),
+										FText(),
+										FSlateIcon(),
+										FUIAction(
+											LoudnessMetric.OnShowValueToggleRequested,
+											FCanExecuteAction::CreateLambda([bCanExecute = LoudnessMetric.OnShowValueToggleRequested.IsBound()]() { return bCanExecute; }),
+											FIsActionChecked::CreateLambda([bShowValue = LoudnessMetric.bShowValue]() { return bShowValue.Get(); })),
+										NAME_None,
+										EUserInterfaceActionType::ToggleButton);
+
+									SubMenu.AddMenuEntry(
+										LOCTEXT("LayoutOptions_LoudnessMetric_ShowMeter", "Show Meter"),
+										FText(),
+										FSlateIcon(),
+										FUIAction(
+											LoudnessMetric.OnShowMeterToggleRequested,
+											FCanExecuteAction::CreateLambda([bCanExecute = LoudnessMetric.OnShowMeterToggleRequested.IsBound()]() { return bCanExecute; }),
+											FIsActionChecked::CreateLambda([bShowMeter = LoudnessMetric.bShowMeter]() { return bShowMeter.Get(); })),
+										NAME_None,
+										EUserInterfaceActionType::ToggleButton);
+
+									if (LoudnessMetric.OnColorChanged.IsBound())
+									{
+										SubMenu.AddMenuEntry(
+											LOCTEXT("LayoutOptions_LoudnessMetric_ValueColor", "Value Color"),
+											FText(),
+											FSlateIcon(FAudioWidgetsCoreStyle::StyleName, "AudioWidgetsCoreStyle.ColorMixing"),
+											FExecuteAction::CreateLambda([
+												ParentWidget = FSlateApplication::Get().GetMenuHostWidget().ToWeakPtr(),
+												Color = LoudnessMetric.Color,
+												OnColorChanged = LoudnessMetric.OnColorChanged,
+												OnColorPickerWindowClosed = LoudnessMetric.OnColorPickerWindowClosed]()
+												{
+													FColorPickerArgs PickerArgs(Color.Get(), OnColorChanged);
+													PickerArgs.ParentWidget = ParentWidget.Pin();
+													PickerArgs.bUseAlpha = true;
+													PickerArgs.bIsModal = true;
+													PickerArgs.OnColorPickerWindowClosed = OnColorPickerWindowClosed;
+													OpenColorPicker(PickerArgs);
+												}),
+											NAME_None,
+											EUserInterfaceActionType::Button);
+									}
+
+									if (LoudnessMetric.OnResetColorRequested.IsBound())
+									{
+										SubMenu.AddMenuEntry(
+											LOCTEXT("LayoutOptions_LoudnessMetric_ResetColor", "Reset Color"),
+											FText(),
+											FSlateIcon(FAudioWidgetsCoreStyle::StyleName, "AudioWidgetsCoreStyle.Reset"),
+											LoudnessMetric.OnResetColorRequested,
+											NAME_None,
+											EUserInterfaceActionType::Button);
+									}
+
+									SubMenu.EndSection();
+								}));
+					}
+
+					MenuBuilder.EndSection();
+				}
+
+				return MenuBuilder.MakeWidget();
+			};
+
+		ToolBarBuilder.AddComboButton(
+			FUIAction(),
+			FOnGetContent::CreateLambda(MakeOptionsMenu),
+			LOCTEXT("OptionsButtonDisplayName", "Options"),
+			FText(),
+			FSlateIcon(FAudioWidgetsCoreStyle::StyleName, "AudioWidgetsCoreStyle.Settings"));
+
+		// Ensure collection views have been synchronized at least once before displaying UI:
+		ValuesCollectionView->RefreshVisibleItems();
+		MetersCollectionView->RefreshVisibleItems();
+
+		const auto GetTimerPanelVisibility = [bIsVisible = TimerPanelParams.bIsVisible]() { return (bIsVisible.Get()) ? EVisibility::Visible : EVisibility::Collapsed; };
+
+		constexpr float MeterWidth = 79.0f;
+		constexpr float MeterMinHeight = 160.0f;
+
+		TSharedRef<SLoudnessMetricTileView> MeterTileView = SNew(SLoudnessMetricTileView,
+				FOnMakeTileInnerContent::CreateStatic(&MakeMeterTileInnerContent, LoudnessScaleParams),
+				SLoudnessMetricTileView::FOnSetShowItem::CreateStatic(&SetShowLoudnessMetricMeter),
+				SLoudnessMetricTileView::FOnReorderItems::CreateSP(MetersCollectionView, &FLoudnessMetricCollectionView::DraggedItemReorder))
+			.OnGenerateTile_Static(&MakeMeterValueTileWidget)
+			.ListItemsSource(MetersCollectionView->GetVisibleItems())
+			.ItemWidth(MeterWidth)
+			.ItemAlignment(EListItemAlignment::CenterAligned)
+			.SelectionMode(ESelectionMode::Single)
+			.ScrollbarVisibility(EVisibility::Collapsed)
+			.OnMouseButtonClick_Lambda([](FLoudnessMetricRef InItem)
+			{
+				InItem->OnClicked.ExecuteIfBound();
+			});
+
+		MeterTileView->SetItemHeight(TAttribute<float>::CreateLambda([&TileView = *MeterTileView]()
+			{
+				const int32 NumItems = TileView.GetItems().Num();
+				const FVector2f PanelSize = TileView.GetTickSpaceGeometry().GetLocalSize();
+				const int32 NumItemsPerLine = FMath::Max(1, FMath::FloorToInt(PanelSize.X / MeterWidth));
+				const int32 NumLines = (NumItems + NumItemsPerLine - 1) / NumItemsPerLine;
+				const int32 HeightDivisor = FMath::Max(1, NumLines);
+				return FMath::Max(PanelSize.Y / HeightDivisor, MeterMinHeight);
+			}));
+
+		const TAttribute<EVisibility> GetNumericValuesPanelVisibility = TAttribute<EVisibility>::CreateSPLambda(ValuesCollectionView,
+			[&InValuesCollectionView = ValuesCollectionView.Get(), MeterTileViewWeak = MeterTileView.ToWeakPtr()]()
+			{
+				if (TSharedPtr<SLoudnessMetricTileView> PinnedMeterTileView = MeterTileViewWeak.Pin())
+				{
+					if (PinnedMeterTileView->HasDragDropOperation())
+					{
+						// Never hide the numeric values panel if a meter is being dragged so that the meter can be dropped onto it.
+						return EVisibility::Visible;
+					}
+				}
+				return (InValuesCollectionView.HasVisibleItems()) ? EVisibility::Visible : EVisibility::Collapsed;
+			});
+
+		return SNew(SBorder)
+			.Padding(0.0f)
+			.BorderImage(&BackgroundBrush)
+			[
+				SNew(SVerticalBox)
+					+ SVerticalBox::Slot()
+					.AutoHeight()
+					[
+						SNew(SOverlay)
+							+ SOverlay::Slot()
+							[
+								SNew(SColorBlock)
+									.Color(FStyleColors::Panel.GetSpecifiedColor())
+							]
+							+ SOverlay::Slot()
+							.HAlign(EHorizontalAlignment::HAlign_Right)
+							[
+								ToolBarBuilder.MakeWidget()
+							]
+					]
+					+ SVerticalBox::Slot()
+					.HAlign(EHorizontalAlignment::HAlign_Center)
+					.AutoHeight()
+					[
+						SNew(SBox)
+							.Padding(0.0f, 10.0f)
+							.Visibility_Lambda(GetTimerPanelVisibility)
+							[
+								SNew(SHorizontalBox)
+									+ SHorizontalBox::Slot()
+									.AutoWidth()
+									[
+										SNew(SBorder)
+											.Padding(8.0f, 4.0f)
+											[
+												SNew(STextBlock)
+													.Font(FStyleDefaults::GetFontInfo(10))
+													.Text_Lambda([AnalysisTime = TimerPanelParams.AnalysisTime]() { return FText::AsTimespan(AnalysisTime.Get()); })
+											]
+									]
+									+ SHorizontalBox::Slot()
+									.AutoWidth()
+									[
+										SNew(SButton)
+											.ButtonStyle(&FAppStyle::Get().GetWidgetStyle<FButtonStyle>("SimpleButton"))
+											.OnClicked(TimerPanelParams.OnResetButtonClicked)
+											[
+												SNew(SImage)
+													.Image(FAudioWidgetsCoreStyle::Get().GetBrush("AudioWidgetsCoreStyle.Reset"))
+													.ColorAndOpacity(FSlateColor::UseForeground())
+													.ToolTipText(LOCTEXT("ResetButtonToolTip", "Reset Long Term loudness analysis"))
+											]
+									]
+							]
+					]
+					+ SVerticalBox::Slot()
+					.AutoHeight()
+					.Padding(8.0f, 2.0f)
+					[
+						SNew(SSeparator)
+							.Visibility_Lambda(GetTimerPanelVisibility)
+							.SeparatorImage(FAppStyle::Get().GetBrush("WhiteBrush"))
+							.Orientation(Orient_Horizontal)
+							.Thickness(1.0f)
+							.ColorAndOpacity(FLinearColor::White.CopyWithNewOpacity(0.25f))
+					]
+					+ SVerticalBox::Slot()
+					.AutoHeight()
+					[
+						SNew(SLoudnessMetricTileView,
+								FOnMakeTileInnerContent::CreateStatic(&MakeValueTileInnerContent, LoudnessScaleParams),
+								SLoudnessMetricTileView::FOnSetShowItem::CreateStatic(&SetShowLoudnessMetricValue),
+								SLoudnessMetricTileView::FOnReorderItems::CreateSP(ValuesCollectionView, &FLoudnessMetricCollectionView::DraggedItemReorder))
+							.Visibility(GetNumericValuesPanelVisibility)
+							.OnGenerateTile_Static(&MakeNumericValueTileWidget)
+							.ListItemsSource(ValuesCollectionView->GetVisibleItems())
+							.ItemHeight(60.0f)
+							.ItemWidth(60.0f)
+							.ItemAlignment(EListItemAlignment::CenterAligned)
+							.SelectionMode(ESelectionMode::Single)
+							.OnMouseButtonClick_Lambda([](FLoudnessMetricRef InItem)
+							{
+								InItem->OnClicked.ExecuteIfBound();
+							})
+					]
+					+ SVerticalBox::Slot()
+					.AutoHeight()
+					.Padding(8.0f, 2.0f)
+					[
+						SNew(SSeparator)
+							.Visibility(GetNumericValuesPanelVisibility)
+							.SeparatorImage(FAppStyle::Get().GetBrush("WhiteBrush"))
+							.Orientation(Orient_Horizontal)
+							.Thickness(1.0f)
+							.ColorAndOpacity(FLinearColor::White.CopyWithNewOpacity(0.25f))
+					]
+					+ SVerticalBox::Slot()
+					.FillHeight(1.0f)
+					[
+						MeterTileView
+					]
+			];
+	}
+
+	void FLoudnessMeterWidgetView::AddLoudnessMetric(const FLoudnessMetric& LoudnessMetric)
+	{
+		FLoudnessMetricRef LoudnessMetricRef = MakeShared<FLoudnessMetric>(LoudnessMetric);
+		LoudnessMetrics->Add(LoudnessMetricRef);
+	}
+
+	void FLoudnessMeterWidgetView::RefreshVisibleLoudnessMetrics()
+	{
+		ValuesCollectionView->RefreshVisibleItems();
+		MetersCollectionView->RefreshVisibleItems();
+	}
+
+	FLoudnessMeterWidgetView::FLoudnessMetricCollectionView::FLoudnessMetricCollectionView(TSharedRef<FLoudnessMetricRefArray> InLoudnessMetrics, TFunction<bool(const FLoudnessMetricRef&)> InVisibilityPredicate)
+		: LoudnessMetrics(InLoudnessMetrics)
+		, VisibilityPredicate(InVisibilityPredicate)
+		, VisibleItems(MakeShared<FLoudnessMetricRefArray>())
+	{
+	}
+
+	void FLoudnessMeterWidgetView::FLoudnessMetricCollectionView::RefreshVisibleItems()
+	{
+		const auto GetDesiredVisibleItems = [this]() -> TArray<FLoudnessMetricRef>
+			{
+				const int64 PermutationIndex = DisplayOrderParams.PermutationIndex.Get();
+				const int32 NumLoudnessMetrics = LoudnessMetrics->Num();
+				const TArray<uint8> PermutationArray = GetPermutationArray(PermutationIndex, NumLoudnessMetrics);
+
+				TArray<FLoudnessMetricRef> DesiredVisibleItems;
+				DesiredVisibleItems.Reserve(NumLoudnessMetrics);
+				for (uint8 LoudnessMetricIndex : PermutationArray)
+				{
+					const FLoudnessMetricRef& LoudnessMetricRef = (*LoudnessMetrics)[LoudnessMetricIndex];
+					if (VisibilityPredicate(LoudnessMetricRef))
+					{
+						DesiredVisibleItems.Add(LoudnessMetricRef);
+					}
+				}
+				return DesiredVisibleItems;
+			};
+
+		const TArray<FLoudnessMetricRef> DesiredVisibleItems = GetDesiredVisibleItems();
+
+		// Sync VisibleItems to match DesiredVisibleItems:
+
+		const int32 NumDesired = DesiredVisibleItems.Num();
+		if (VisibleItems->Num() > NumDesired)
+		{
+			VisibleItems->RemoveAt(NumDesired, VisibleItems->Num() - NumDesired);
+		}
+
+		for (int32 Index = 0; Index < NumDesired; Index++)
+		{
+			const FLoudnessMetricRef& LoudnessMetricRef = DesiredVisibleItems[Index];
+
+			if (Index >= VisibleItems->Num())
+			{
+				VisibleItems->Add(LoudnessMetricRef);
+			}
+			else if ((*VisibleItems)[Index] != LoudnessMetricRef)
+			{
+				VisibleItems->RemoveAt(Index);
+				VisibleItems->EmplaceAt(Index, LoudnessMetricRef);
+			}
+		}
+	}
+
+	void FLoudnessMeterWidgetView::FLoudnessMetricCollectionView::DraggedItemReorder(const FLoudnessMetricRef& ItemDroppedOnto, bool bInsertAfter, const FLoudnessMetricRef& DragDropItem)
+	{
+		// Find the given items in the array of all loudness metrics (these are unchanging indexes that uniquely identify the loudness metrics):
+		const int32 DroppedOntoLoudnessMetricIndex = LoudnessMetrics->Find(ItemDroppedOnto);
+		const int32 DragDropLoudnessMetricIndex = LoudnessMetrics->Find(DragDropItem);
+		if (ensure(DroppedOntoLoudnessMetricIndex != INDEX_NONE && DragDropLoudnessMetricIndex != INDEX_NONE && DragDropLoudnessMetricIndex != DroppedOntoLoudnessMetricIndex))
+		{
+			// We get the current PermutationArray and then modify it as required.
+			const uint64 CurrentPermutationIndex = DisplayOrderParams.PermutationIndex.Get();
+			TArray<uint8> PermutationArray = GetPermutationArray(CurrentPermutationIndex, LoudnessMetrics->Num());
+
+			// Remove dragged item from current position:
+			PermutationArray.Remove(DragDropLoudnessMetricIndex);
+
+			const int32 DroppedOntoPosition = PermutationArray.Find(DroppedOntoLoudnessMetricIndex);
+			if (ensure(DroppedOntoPosition != INDEX_NONE))
+			{
+				// Place dragged item at new position:
+				const int32 InsertPosition = bInsertAfter ? DroppedOntoPosition + 1 : DroppedOntoPosition;
+				PermutationArray.EmplaceAt(InsertPosition, DragDropLoudnessMetricIndex);
+
+				// Get the new permutation index with the reordering applied:
+				const uint64 NewPermutationIndex = GetPermutationIndex(PermutationArray);
+				if (NewPermutationIndex != CurrentPermutationIndex && DisplayOrderParams.OnPermutationChanged.IsBound())
+				{
+					// Raise request to change ordering with new PermutationIndex:
+					DisplayOrderParams.OnPermutationChanged.Execute(NewPermutationIndex);
+
+					// Proactive refresh:
+					RefreshVisibleItems();
+				}
+			}
+		}
+	}
+
+	uint64 FLoudnessMeterWidgetView::FLoudnessMetricCollectionView::GetPermutationIndex(TArrayView<uint8> PermutationArray)
+	{
+		using namespace Permutations;
+
+		const uint64 NumPermuations = Factorial(PermutationArray.Num());
+		const uint64 MaxPermutationIndex = NumPermuations - 1;
+
+		// We alter the returned index so that the permutation (0,1,2,...) is represented as PermutationIndex zero:
+		return MaxPermutationIndex - AnalyzePermutation(PermutationArray);
+	}
+
+	TArray<uint8> FLoudnessMeterWidgetView::FLoudnessMetricCollectionView::GetPermutationArray(uint64 PermutationIndex, uint8 NumValues)
+	{
+		using namespace Permutations;
+
+		const uint64 NumPermuations = Factorial(NumValues);
+		const uint64 MaxPermutationIndex = NumPermuations - 1;
+
+		// Ensure valid bounds:
+		PermutationIndex = PermutationIndex % NumPermuations;
+
+		// We use a reversed index so that a given PermutationIndex of zero will result in the permutation (0,1,2,...):
+		return GeneratePermutation(MaxPermutationIndex - PermutationIndex, NumValues);
+	}
+} // namespace AudioWidgetsCore
+
+#undef LOCTEXT_NAMESPACE
diff --git a/Engine/Plugins/Runtime/AudioWidgets/Source/AudioWidgetsCore/Private/Permutations.cpp b/Engine/Plugins/Runtime/AudioWidgets/Source/AudioWidgetsCore/Private/Permutations.cpp
new file mode 100644
--- /dev/null
+++ b/Engine/Plugins/Runtime/AudioWidgets/Source/AudioWidgetsCore/Private/Permutations.cpp
@@ -0,0 +1,51 @@
+// Copyright Epic Games, Inc. All Rights Reserved.
+
+#include "Permutations.h"
+
+#include "Algo/MaxElement.h"
+
+namespace AudioWidgetsCore::Permutations
+{
+	uint64 AnalyzePermutation(TArrayView<uint8> PermutationArray)
+	{
+		check(PermutationArray.Num() >= 0);
+		check(PermutationArray.Num() <= MaxFactorial);
+
+		// See Knuth, Vol. 2, 3rd ed., Section 3.3.2, pp. 65-66, Algorithm P (Analyze a permutation).
+
+		uint64 F = 0;
+		for (uint8 R = PermutationArray.Num(); R > 1; R--)
+		{
+			PermutationArray.LeftInline(R);
+			const uint64 S = Algo::MaxElement(PermutationArray) - PermutationArray.GetData();
+			F = R * F + S;
+			Swap(PermutationArray[R - 1], PermutationArray[S]);
+		}
+		return F;
+	}
+
+	TArray<uint8> GeneratePermutation(uint64 PermutationID, uint8 NumValues)
+	{
+		check(NumValues <= MaxFactorial);
+		check(PermutationID < Factorial(NumValues));
+
+		// Algorithm P, reverse
+
+		TArray<uint8> PermutationArray;
+		PermutationArray.Reserve(NumValues);
+		while (PermutationArray.Num() < NumValues)
+		{
+			PermutationArray.Add(PermutationArray.Num());
+		}
+
+		uint64 F = PermutationID;
+		for (uint8 R = 2; R <= NumValues; R++)
+		{
+			const uint64 S = F % R;
+			F /= R;
+			PermutationArray.Swap(R - 1, S);
+		}
+		return PermutationArray;
+	}
+
+} // namespace AudioWidgetsCore::Permutations
diff --git a/Engine/Plugins/Runtime/AudioWidgets/Source/AudioWidgetsCore/Private/Permutations.h b/Engine/Plugins/Runtime/AudioWidgets/Source/AudioWidgetsCore/Private/Permutations.h
new file mode 100644
--- /dev/null
+++ b/Engine/Plugins/Runtime/AudioWidgets/Source/AudioWidgetsCore/Private/Permutations.h
@@ -0,0 +1,46 @@
+// Copyright Epic Games, Inc. All Rights Reserved.
+
+#pragma once
+
+#include "Containers/ArrayView.h"
+
+namespace AudioWidgetsCore::Permutations
+{
+	constexpr uint8 MaxFactorial = 20;
+
+	/**
+	 * Compute exact factorial. Only valid for small non-negative integers.
+	 */
+	constexpr uint64 Factorial(uint8 N)
+	{
+		if (!ensure(N <= MaxFactorial))
+		{
+			return UINT64_MAX;
+		}
+
+		uint64 F = 1;
+		for (uint8 R = N; R > 1; R--)
+		{
+			F *= R;
+		}
+		return F;
+	}
+
+	/**
+	 * Get an integer that uniquely represents the given permutation.
+	 * 
+	 * @param PermutationArray The elements of the array should be distinct values in the interval [0, N-1] where N == PermutationArray.Num(). On return, the array will have been sorted.
+	 * @return A value in the interval [0, N!-1] where N == PermutationArray.Num(). If N > MaxFactorial, the result is undefined.
+	 */
+	uint64 AnalyzePermutation(TArrayView<uint8> PermutationArray);
+
+	/**
+	 * Create the permutation that is uniquely represented by the given integer.
+	 * 
+	 * @param PermutationID A value in the interval [0, N!-1] where N == NumValues.
+	 * @param NumValues The desired number of values in the permutation, must be no greater than MaxFactorial.
+	 * @return An array of distinct values with length == NumValues. If NumValues > MaxFactorial, the result is undefined.
+	 */
+	TArray<uint8> GeneratePermutation(uint64 PermutationID, uint8 NumValues);
+
+} // namespace AudioWidgetsCore::Permutations
diff --git a/Engine/Plugins/Runtime/AudioWidgets/Source/AudioWidgetsCore/Private/SDragReorderableTileView.cpp b/Engine/Plugins/Runtime/AudioWidgets/Source/AudioWidgetsCore/Private/SDragReorderableTileView.cpp
new file mode 100644
--- /dev/null
+++ b/Engine/Plugins/Runtime/AudioWidgets/Source/AudioWidgetsCore/Private/SDragReorderableTileView.cpp
@@ -0,0 +1,45 @@
+// Copyright Epic Games, Inc. All Rights Reserved.
+
+#include "SDragReorderableTileView.h"
+
+#include "Styling/StyleColors.h"
+
+namespace AudioWidgetsCore
+{
+	int32 DrawDropIndicator(const FGeometry& AllottedGeometry, FSlateWindowElementList& OutDrawElements, int32 LayerId, const FWidgetStyle& InWidgetStyle, EHorizontalAlignment HorizontalAlignment)
+	{
+		const FVector2f TileSize = AllottedGeometry.GetLocalSize();
+
+		const float DropIndicatorTranslateX = [TileWidth = TileSize.X, HorizontalAlignment]()
+			{
+				switch (HorizontalAlignment)
+				{
+				default:
+				case HAlign_Fill:
+				case HAlign_Left:
+					return 0.0f;
+				case HAlign_Center:
+					return 0.5f * TileWidth;
+				case HAlign_Right:
+					return TileWidth;
+				}
+			}();
+
+		constexpr float DropIndicatorScale = 1.0f;
+		const FVector2f DropIndicatorSize(1.0f, TileSize.Y);
+		const FVector2f DropIndicatorTranslation(DropIndicatorTranslateX, 0.0f);
+		const FGeometry DropIndicatorGeometry = AllottedGeometry.MakeChild(DropIndicatorSize, FSlateLayoutTransform(DropIndicatorScale, DropIndicatorTranslation));
+
+		FSlateDrawElement::MakeBox
+		(
+			OutDrawElements,
+			LayerId++,
+			DropIndicatorGeometry.ToPaintGeometry(),
+			FAppStyle::GetBrush("PropertyEditor.VerticalDottedLine"),
+			ESlateDrawEffect::None,
+			FStyleColors::Select.GetColor(InWidgetStyle) * InWidgetStyle.GetColorAndOpacityTint()
+		);
+
+		return LayerId;
+	}
+} // namespace AudioWidgetsCore
diff --git a/Engine/Plugins/Runtime/AudioWidgets/Source/AudioWidgetsCore/Private/SDragReorderableTileView.h b/Engine/Plugins/Runtime/AudioWidgets/Source/AudioWidgetsCore/Private/SDragReorderableTileView.h
new file mode 100644
--- /dev/null
+++ b/Engine/Plugins/Runtime/AudioWidgets/Source/AudioWidgetsCore/Private/SDragReorderableTileView.h
@@ -0,0 +1,304 @@
+// Copyright Epic Games, Inc. All Rights Reserved.
+
+#pragma once
+
+#include "Framework/SlateDelegates.h"
+#include "Widgets/Views/STileView.h"
+
+namespace AudioWidgetsCore
+{
+	/**
+	 * Function to draw a dotted line representing where the current drag drop would place the dragged item.
+	 */
+	int32 DrawDropIndicator(const FGeometry& AllottedGeometry, FSlateWindowElementList& OutDrawElements, int32 LayerId, const FWidgetStyle& InWidgetStyle, EHorizontalAlignment HorizontalAlignment);
+
+	template<typename DragDropOperationT>
+	class SDragReorderableTileView;
+
+	/**
+	 * Abstract base class, represents an item that is being dragged from a given source tile view.
+	 * Uses Curiously Recurring Template Pattern so that FDragDropOperation type checking operations can be used.
+	 */
+	template<typename ItemT, typename DragDropOperationT>
+	class TItemDragDropOperation : public FDragDropOperation
+	{
+	public:
+		using ItemType = ItemT;
+		using STileViewType = SDragReorderableTileView<DragDropOperationT>;
+
+		TItemDragDropOperation(TSharedRef<STileViewType> InSourceTileView)
+			: WeakSourceTileView(InSourceTileView.ToWeakPtr())
+		{
+		}
+
+		virtual const ItemType& GetItem() const = 0;
+
+		TSharedPtr<STileViewType> GetSourceTileView() const { return WeakSourceTileView.Pin(); }
+
+		virtual FCursorReply OnCursorQuery() override { return FCursorReply::Cursor(EMouseCursor::GrabHand); }
+
+	private:
+		const TWeakPtr<STileViewType> WeakSourceTileView;
+	};
+
+	template<typename T>
+	concept ItemDragDropOperation =
+		requires { typename T::ItemType; };// && std::derived_from<T, TItemDragDropOperation<typename T::ItemType, T>>; // Can uncomment when all platform stdlibs have <concepts> header.
+
+
+	/**
+	 * Abstract base class for TileView widget that allow its tiles to be dragged and dropped for reordering etc.
+	 * Dropping of a dragged item onto the tile view is supported if the tile view is currently empty.
+	 */
+	template<typename DragDropOperationT>
+	class SDragReorderableTileView : public STileView<typename DragDropOperationT::ItemType>
+	{
+	public:
+		using Base = STileView<typename DragDropOperationT::ItemType>;
+		using ItemType = typename DragDropOperationT::ItemType;
+
+		DECLARE_DELEGATE_TwoParams(FOnSetShowItem, const ItemType&, bool);
+		DECLARE_DELEGATE_ThreeParams(FOnReorderItems, const ItemType&/*ItemDroppedOnto*/, bool/*bInsertAfter*/, const ItemType&/*DragDropItem*/);
+
+		void Construct(const typename Base::FArguments& InArgs, FOnSetShowItem InOnSetShowItem, FOnReorderItems InOnReorderItems)
+		{
+			Base::Construct(InArgs);
+
+			OnSetShowItem = InOnSetShowItem;
+			OnReorderItems = InOnReorderItems;
+		}
+
+		virtual TSharedPtr<DragDropOperationT> CreateDragDropOperation() = 0;
+
+		void SetShowItem(const ItemType& Item, bool bShow)
+		{
+			OnSetShowItem.ExecuteIfBound(Item, bShow);
+		}
+
+		void ReorderItems(const ItemType& ItemDroppedOnto, bool bInsertAfter, const ItemType& DragDropItem)
+		{
+			ensure(ItemDroppedOnto != DragDropItem);
+			OnReorderItems.ExecuteIfBound(ItemDroppedOnto, bInsertAfter, DragDropItem);
+		}
+
+		// Begin SWidget overrides.
+		virtual FReply OnDragOver(const FGeometry& MyGeometry, const FDragDropEvent& DragDropEvent) override
+		{
+			if (TSharedPtr<DragDropOperationT> DragDropOp = DragDropEvent.GetOperationAs<DragDropOperationT>())
+			{
+				// Respond to drag drop only if we have no items, otherwise the tile will handle it:
+				if (Base::GetItems().IsEmpty())
+				{
+					bPaintDropIndicator = true;
+					return FReply::Handled();
+				}
+			}
+			return FReply::Unhandled();
+		}
+
+		virtual void OnDragLeave(FDragDropEvent const& DragDropEvent) override
+		{
+			bPaintDropIndicator = false;
+		}
+
+		virtual FReply OnDrop(const FGeometry& MyGeometry, const FDragDropEvent& DragDropEvent) override
+		{
+			bPaintDropIndicator = false;
+
+			if (TSharedPtr<DragDropOperationT> DragDropOp = DragDropEvent.GetOperationAs<DragDropOperationT>())
+			{
+				// Respond to drag drop only if we have no items, otherwise the tile will handle it:
+				if (Base::GetItems().IsEmpty())
+				{
+					if (TSharedPtr<SDragReorderableTileView> DragSourceTileView = DragDropOp->GetSourceTileView())
+					{
+						const auto& DragDropItem = DragDropOp->GetItem();
+
+						// Move between tile views:
+						DragSourceTileView->SetShowItem(DragDropItem, false);
+						SetShowItem(DragDropItem, true);
+
+						return FReply::Handled();
+					}
+				}
+			}
+			return FReply::Unhandled();
+		}
+
+		virtual int32 OnPaint(const FPaintArgs& Args, const FGeometry& AllottedGeometry, const FSlateRect& MyCullingRect, FSlateWindowElementList& OutDrawElements, int32 LayerId, const FWidgetStyle& InWidgetStyle, bool bParentEnabled) const override
+		{
+			LayerId = Base::OnPaint(Args, AllottedGeometry, MyCullingRect, OutDrawElements, LayerId, InWidgetStyle, bParentEnabled);
+
+			if (bPaintDropIndicator)
+			{
+				LayerId = DrawDropIndicator(AllottedGeometry, OutDrawElements, LayerId, InWidgetStyle, EHorizontalAlignment::HAlign_Center);
+			}
+
+			return LayerId;
+		}
+		// End SWidget overrides
+
+	private:
+		FOnSetShowItem OnSetShowItem;
+		FOnReorderItems OnReorderItems;
+		bool bPaintDropIndicator = false;
+	};
+
+	/**
+	 * Base class for tiles displayed by SDragReorderableTileView.
+	 */
+	template<ItemDragDropOperation DragDropOperationT>
+	class SDragReorderableTile : public STableRow<typename DragDropOperationT::ItemType>
+	{
+	public:
+		using Base = STableRow<typename DragDropOperationT::ItemType>;
+		using ItemType = typename DragDropOperationT::ItemType;
+		using STileViewType = SDragReorderableTileView<DragDropOperationT>;
+
+		bool IsBeingDragged() const { return WeakDragDropOperation.IsValid(); }
+
+	protected:
+		// Begin SWidget overrides.
+		virtual FReply OnDragDetected(const FGeometry& MyGeometry, const FPointerEvent& PointerEvent) override
+		{
+			const FReply BaseReply = Base::OnDragDetected(MyGeometry, PointerEvent);
+			if (!BaseReply.IsEventHandled() && PointerEvent.IsMouseButtonDown(EKeys::LeftMouseButton))
+			{
+				if (TSharedPtr<STileViewType> OwnerTileView = GetPinnedTileView())
+				{
+					if (TSharedPtr<FDragDropOperation> DragDropOperation = OwnerTileView->CreateDragDropOperation())
+					{
+						// Hold a weak pointer to maintain awareness of whether this tile is currently being dragged:
+						WeakDragDropOperation = DragDropOperation;
+
+						return FReply::Handled().BeginDragDrop(DragDropOperation.ToSharedRef());
+					}
+				}
+			}
+			return BaseReply;
+		}
+
+		virtual FReply OnDragOver(const FGeometry& MyGeometry, const FDragDropEvent& DragDropEvent) override
+		{
+			// Setting STableRow::ItemDropZone TOptional will enable painting the drop indicator:
+			Base::SetItemDropZone(CanAcceptDrop(MyGeometry, DragDropEvent));
+
+			return FReply::Handled();
+		}
+
+		virtual FReply OnDrop(const FGeometry& MyGeometry, const FDragDropEvent& DragDropEvent) override
+		{
+			// Clear ItemDropZone to no longer paint the drop indicator:
+			Base::SetItemDropZone(TOptional<EItemDropZone>());
+
+			if (TOptional<EItemDropZone> AcceptableDropZone = CanAcceptDrop(MyGeometry, DragDropEvent))
+			{
+				return AcceptDrop(DragDropEvent, *AcceptableDropZone);
+			}
+
+			return FReply::Unhandled();
+		}
+		// End SWidget overrides
+
+		// Begin STableRow overrides.
+		virtual int32 OnPaintDropIndicator(EItemDropZone InItemDropZone, const FPaintArgs& Args, const FGeometry& AllottedGeometry, const FSlateRect& MyCullingRect, FSlateWindowElementList& OutDrawElements, int32 LayerId, const FWidgetStyle& InWidgetStyle, bool bParentEnabled) const override
+		{
+			const EHorizontalAlignment HorizontalAlignment = (InItemDropZone == EItemDropZone::BelowItem) ? EHorizontalAlignment::HAlign_Right : EHorizontalAlignment::HAlign_Left;
+			return DrawDropIndicator(AllottedGeometry, OutDrawElements, LayerId, InWidgetStyle, HorizontalAlignment);
+		}
+		// End STableRow overrides.
+
+	private:
+		TOptional<EItemDropZone> CanAcceptDrop(const FGeometry& MyGeometry, const FDragDropEvent& DragDropEvent)
+		{
+			if (TSharedPtr<DragDropOperationT> DragDropOp = DragDropEvent.GetOperationAs<DragDropOperationT>())
+			{
+				if (TSharedPtr<STileViewType> OwnerTileView = GetPinnedTileView())
+				{
+					const EItemDropZone ItemHoverZone = GetItemHoverZone(MyGeometry, DragDropEvent);
+					if (IsValidDrop(OwnerTileView.ToSharedRef(), ItemHoverZone, DragDropOp.ToSharedRef()))
+					{
+						return ItemHoverZone;
+					}
+				}
+			}
+			return NullOpt;
+		}
+
+		bool IsValidDrop(const TSharedRef<STileViewType>& OwnerTileView, EItemDropZone ItemHoverZone, TSharedRef<DragDropOperationT> DragDropOp) const
+		{
+			const auto& DragDropItem = DragDropOp->GetItem();
+
+			TConstArrayView<ItemType> Items = OwnerTileView->GetItems();
+			const int32 DragDroppedIndex = Items.Find(DragDropItem);
+			if (DragDroppedIndex == INDEX_NONE)
+			{
+				// Item does not exist in the current tile view so it is always valid.
+				return true;
+			}
+
+			if (DragDropOp->GetSourceTileView() != OwnerTileView)
+			{
+				// Don't alow drag from another tile view as the current tile view already contains the loudness metric.
+				return false;
+			}
+
+			if (const ItemType* ItemPtr = Base::GetItemForThis(OwnerTileView))
+			{
+				const int32 DroppedOntoIndex = Items.Find(*ItemPtr);
+				if (DroppedOntoIndex != INDEX_NONE && DragDroppedIndex != DroppedOntoIndex)
+				{
+					const bool bInsertAfter = (ItemHoverZone == EItemDropZone::BelowItem);
+					const int32 InsertIndex = (bInsertAfter) ? DroppedOntoIndex + 1 : DroppedOntoIndex - 1;
+
+					// Return valid if insert position causes change:
+					return (InsertIndex != DragDroppedIndex);
+				}
+			}
+			return false;
+		}
+
+		FReply AcceptDrop(const FDragDropEvent& DragDropEvent, EItemDropZone ItemHoverZone)
+		{
+			if (TSharedPtr<DragDropOperationT> DragDropOp = DragDropEvent.GetOperationAs<DragDropOperationT>())
+			{
+				if (TSharedPtr<STileViewType> OwnerTileView = GetPinnedTileView())
+				{
+					if (const ItemType* ItemPtr = Base::GetItemForThis(OwnerTileView.ToSharedRef()))
+					{
+						const auto& DragDropItem = DragDropOp->GetItem();
+
+						// Place the dropped item as implied by ItemHoverZone:
+						const bool bInsertAfter = (ItemHoverZone == EItemDropZone::BelowItem);
+						OwnerTileView->ClearSelection();
+						OwnerTileView->ReorderItems(*ItemPtr, bInsertAfter, DragDropItem);
+
+						TSharedPtr<STileViewType> DragSourceTileView = DragDropOp->GetSourceTileView();
+						if (DragSourceTileView.IsValid() && DragSourceTileView != OwnerTileView)
+						{
+							// Item was dragged from other tile view, move it between tile views:
+							DragSourceTileView->SetShowItem(DragDropItem, false);
+							OwnerTileView->SetShowItem(DragDropItem, true);
+						}
+
+						return FReply::Handled();
+					}
+				}
+			}
+			return FReply::Unhandled();
+		}
+
+		TSharedPtr<STileViewType> GetPinnedTileView()
+		{
+			return StaticCastSharedPtr<STileViewType>(Base::OwnerTablePtr.Pin());
+		}
+
+		EItemDropZone GetItemHoverZone(const FGeometry& MyGeometry, const FDragDropEvent& DragDropEvent)
+		{
+			const FVector2f LocalPointerPos = MyGeometry.AbsoluteToLocal(DragDropEvent.GetScreenSpacePosition());
+			return Base::ZoneFromPointerPosition(LocalPointerPos, MyGeometry.GetLocalSize(), EOrientation::Orient_Horizontal);
+		}
+
+		TWeakPtr<FDragDropOperation> WeakDragDropOperation; // Will be set when a drag is detected.
+	};
+} // namespace AudioWidgetsCore
diff --git a/Engine/Plugins/Runtime/AudioWidgets/Source/AudioWidgetsCore/Public/AudioWidgetsCoreStyle.h b/Engine/Plugins/Runtime/AudioWidgets/Source/AudioWidgetsCore/Public/AudioWidgetsCoreStyle.h
new file mode 100644
--- /dev/null
+++ b/Engine/Plugins/Runtime/AudioWidgets/Source/AudioWidgetsCore/Public/AudioWidgetsCoreStyle.h
@@ -0,0 +1,25 @@
+// Copyright Epic Games, Inc. All Rights Reserved.
+
+#pragma once
+
+#include "Styling/SlateStyle.h"
+
+#define UE_API AUDIOWIDGETSCORE_API
+
+/**
+ * Slate style set that defines all the styles for audio widgets core
+ */
+class FAudioWidgetsCoreStyle : public FSlateStyleSet
+{
+public:
+	/** Access the singleton instance for this style set */
+	UE_API static const FAudioWidgetsCoreStyle& Get();
+
+	inline static const FName StyleName = "AudioWidgetsCoreStyle";
+
+private:
+	FAudioWidgetsCoreStyle();
+	~FAudioWidgetsCoreStyle();
+};
+
+#undef UE_API
diff --git a/Engine/Plugins/Runtime/AudioWidgets/Source/AudioWidgetsCore/Public/LoudnessMeterSettings.h b/Engine/Plugins/Runtime/AudioWidgets/Source/AudioWidgetsCore/Public/LoudnessMeterSettings.h
new file mode 100644
--- /dev/null
+++ b/Engine/Plugins/Runtime/AudioWidgets/Source/AudioWidgetsCore/Public/LoudnessMeterSettings.h
@@ -0,0 +1,64 @@
+// Copyright Epic Games, Inc. All Rights Reserved.
+
+#pragma once
+
+#include "Styling/StyleColors.h"
+
+#include "LoudnessMeterSettings.generated.h"
+
+USTRUCT()
+struct FLoudnessMeterDisplayOptions
+{
+	GENERATED_BODY()
+
+	UPROPERTY(EditAnywhere, config, Category = LoudnessMeter)
+	bool bShowValue = false;
+
+	UPROPERTY(EditAnywhere, config, Category = LoudnessMeter)
+	bool bShowMeter = false;
+
+	UPROPERTY(EditAnywhere, config, Category = LoudnessMeter)
+	FLinearColor Color = FStyleColors::AccentGreen.GetSpecifiedColor().CopyWithNewOpacity(0.5f);
+};
+
+USTRUCT()
+struct FLoudnessMeterSettings
+{
+	GENERATED_BODY()
+
+	UPROPERTY(EditAnywhere, config, Category = LoudnessMeter, meta = (ClampMin = "10", ClampMax = "70", UIMin = "10", UIMax = "70"))
+	int LoudnessScaleRange = 60;
+
+	UPROPERTY(EditAnywhere, config, Category = LoudnessMeter, meta = (ClampMin = "-10", ClampMax = "30", UIMin = "-10", UIMax = "30"))
+	int LoudnessScaleOffset = 0;
+
+	UPROPERTY(EditAnywhere, config, Category = LoudnessMeter, meta = (ClampMin = "-60", ClampMax = "30", UIMin = "-60", UIMax = "30"))
+	int LoudnessScaleTarget = -23;
+
+	UPROPERTY(EditAnywhere, config, Category = LoudnessMeter)
+	bool bUseRelativeLoudnessScale = false;
+
+	UPROPERTY(EditAnywhere, config, Category = LoudnessMeter)
+	bool bDisplayAnalysisTimer = true;
+
+	UPROPERTY(EditAnywhere, config, Category = LoudnessMeter)
+	FLoudnessMeterDisplayOptions LongTermLoudness = { .bShowValue = true, .bShowMeter = false };
+
+	UPROPERTY(EditAnywhere, config, Category = LoudnessMeter)
+	FLoudnessMeterDisplayOptions ShortTermLoudness = { .bShowValue = false, .bShowMeter = true };
+
+	UPROPERTY(EditAnywhere, config, Category = LoudnessMeter)
+	FLoudnessMeterDisplayOptions MomentaryLoudness = { .bShowValue = false, .bShowMeter = true };
+
+	UPROPERTY(EditAnywhere, config, Category = LoudnessMeter)
+	FLoudnessMeterDisplayOptions LoudnessRange = { .bShowValue = true, .bShowMeter = false };
+
+	UPROPERTY(EditAnywhere, config, Category = LoudnessMeter)
+	FLoudnessMeterDisplayOptions TruePeak = { .bShowValue = true, .bShowMeter = false };
+
+	UPROPERTY(EditAnywhere, config, Category = LoudnessMeter, meta = (ClampMin = "0", ClampMax = "119", UIMin = "0", UIMax = "119"))
+	uint64 ValuesOrderingPermutation = 0; // Max == 5! - 1
+
+	UPROPERTY(EditAnywhere, config, Category = LoudnessMeter, meta = (ClampMin = "0", ClampMax = "119", UIMin = "0", UIMax = "119"))
+	uint64 MetersOrderingPermutation = 0; // Max == 5! - 1
+};
diff --git a/Engine/Plugins/Runtime/AudioWidgets/Source/AudioWidgetsCore/Public/LoudnessMeterWidgetView.h b/Engine/Plugins/Runtime/AudioWidgets/Source/AudioWidgetsCore/Public/LoudnessMeterWidgetView.h
new file mode 100644
--- /dev/null
+++ b/Engine/Plugins/Runtime/AudioWidgets/Source/AudioWidgetsCore/Public/LoudnessMeterWidgetView.h
@@ -0,0 +1,133 @@
+// Copyright Epic Games, Inc. All Rights Reserved.
+
+#pragma once
+
+#include "Containers/ObservableArray.h"
+#include "Framework/SlateDelegates.h"
+#include "Widgets/SWidget.h"
+#include "Widgets/SWindow.h"
+
+#define UE_API AUDIOWIDGETSCORE_API
+
+namespace AudioWidgetsCore
+{
+	enum class EAudioMeterMetric : uint8
+	{
+		Loudness,
+		Decibels
+	};
+
+	/**
+	 * A widget that displays various loudness metrics as numeric values and/or meters.
+	 */
+	class FLoudnessMeterWidgetView
+	{
+	public:
+		UE_API static FLinearColor GetDefaultMeterColor();
+
+		struct FLoudnessMetric
+		{
+			FName Name;
+			FText DisplayName;
+			TAttribute<EAudioMeterMetric> MeterMetric;
+			TAttribute<TOptional<float>> Value;
+			TAttribute<TOptional<FFloatInterval>> Range;
+			TAttribute<bool> bShowValue = false;
+			TAttribute<bool> bShowMeter = false;
+			FSimpleDelegate OnShowValueToggleRequested;
+			FSimpleDelegate OnShowMeterToggleRequested;
+			TAttribute<FLinearColor> Color = GetDefaultMeterColor();
+			FOnLinearColorValueChanged OnColorChanged;
+			FOnWindowClosed OnColorPickerWindowClosed;
+			FSimpleDelegate OnResetColorRequested;
+			FSimpleDelegate OnClicked;
+
+			bool IsRanged() const { return Range.IsBound() || Range.Get(NullOpt).IsSet(); }
+		};
+
+		struct FTimerPanelParams
+		{
+			TAttribute<FTimespan> AnalysisTime = FTimespan::Zero();
+			FOnClicked OnResetButtonClicked;
+			TAttribute<bool> bIsVisible = false;
+			FSimpleDelegate OnVisibilityToggleRequested;
+		};
+
+		DECLARE_DELEGATE_OneParam(FOnIntValueChanged, int32 /*NewValue*/);
+		DECLARE_DELEGATE_TwoParams(FOnIntValueCommitted, int32 /*NewValue*/, ETextCommit::Type /*CommitType*/);
+
+		struct FLoudnessScaleParams
+		{
+			TAttribute<int32> Range = 60;
+			FOnIntValueChanged OnRangeValueChanged;
+			FOnIntValueCommitted OnRangeValueCommitted;
+
+			TAttribute<int32> Offset = 0;
+			FOnIntValueChanged OnOffsetValueChanged;
+			FOnIntValueCommitted OnOffsetValueCommitted;
+
+			TAttribute<int32> Target = -23;
+			FOnIntValueChanged OnTargetValueChanged;
+			FOnIntValueCommitted OnTargetValueCommitted;
+
+			TAttribute<bool> bRelativeScale = false;
+			FSimpleDelegate OnRelativeScaleToggleRequested;
+		};
+
+		DECLARE_DELEGATE_OneParam(FOnPermutationChanged, uint64 /*NewPermutationIndex*/);
+
+		struct FDisplayOrderParams
+		{
+			TAttribute<uint64> PermutationIndex = 0;
+			FOnPermutationChanged OnPermutationChanged;
+		};
+
+		UE_API FLoudnessMeterWidgetView();
+		UE_API ~FLoudnessMeterWidgetView();
+
+		UE_API void InitTimerPanel(const FTimerPanelParams& InTimerPanelParams);
+		UE_API void InitLoudnessScale(const FLoudnessScaleParams& InLoudnessScaleParams);
+		UE_API void InitValuesDisplayOrder(const FDisplayOrderParams& InDisplayOrderParams);
+		UE_API void InitMetersDisplayOrder(const FDisplayOrderParams& InDisplayOrderParams);
+
+		UE_API TSharedRef<SWidget> MakeWidget() const;
+
+		UE_API void AddLoudnessMetric(const FLoudnessMetric& LoudnessMetric);
+		UE_API void RefreshVisibleLoudnessMetrics();
+
+	private:
+		using FLoudnessMetricRef = TSharedRef<const FLoudnessMetric>;
+		using FLoudnessMetricRefArray = UE::Slate::Containers::TObservableArray<FLoudnessMetricRef>;
+
+		class FLoudnessMetricCollectionView : public TSharedFromThis<FLoudnessMetricCollectionView>
+		{
+		public:
+			FLoudnessMetricCollectionView(TSharedRef<FLoudnessMetricRefArray> InLoudnessMetrics, TFunction<bool(const FLoudnessMetricRef&)> InVisibilityPredicate);
+
+			void SetDisplayOrderParams(const FDisplayOrderParams& InDisplayOrderParams) { DisplayOrderParams = InDisplayOrderParams; }
+			void RefreshVisibleItems();
+			void ResetVisibleItems() { VisibleItems->Reset(); }
+			TSharedRef<FLoudnessMetricRefArray> GetVisibleItems() { return VisibleItems; }
+			bool HasVisibleItems() const { return !VisibleItems->IsEmpty(); }
+
+			void DraggedItemReorder(const FLoudnessMetricRef& ItemDroppedOnto, bool bInsertAfter, const FLoudnessMetricRef& DragDropItem);
+
+		private:
+			static uint64 GetPermutationIndex(TArrayView<uint8> PermutationArray);
+			static TArray<uint8> GetPermutationArray(uint64 PermutationIndex, uint8 NumValues);
+
+			const TSharedRef<FLoudnessMetricRefArray> LoudnessMetrics;
+			const TFunction<bool(const FLoudnessMetricRef&)> VisibilityPredicate;
+			const TSharedRef<FLoudnessMetricRefArray> VisibleItems;
+			FDisplayOrderParams DisplayOrderParams;
+		};
+
+		TSharedRef<FLoudnessMetricRefArray> LoudnessMetrics;
+		TSharedRef<FLoudnessMetricCollectionView> ValuesCollectionView;
+		TSharedRef<FLoudnessMetricCollectionView> MetersCollectionView;
+		FTimerPanelParams TimerPanelParams;
+		FLoudnessScaleParams LoudnessScaleParams;
+	};
+} // namespace AudioWidgetsCore
+
+#undef UE_API
