[AudioWidgets] LoudnessMeter: added logic to visualize (and customize the color of) the Target value in loudness meters.

#Synced-CL 51006836

diff --git a/Engine/Plugins/Runtime/AudioWidgets/Source/AudioWidgets/Private/AudioLoudnessMeter.cpp b/Engine/Plugins/Runtime/AudioWidgets/Source/AudioWidgets/Private/AudioLoudnessMeter.cpp
--- a/Engine/Plugins/Runtime/AudioWidgets/Source/AudioWidgets/Private/AudioLoudnessMeter.cpp
+++ b/Engine/Plugins/Runtime/AudioWidgets/Source/AudioWidgets/Private/AudioLoudnessMeter.cpp
@@ -337,6 +337,19 @@ namespace AudioWidgets
 						SettingsHelper.SaveConfig();
 					});
 
+				LoudnessScaleParams.TargetColor.BindLambda([=]() { return SettingsHelper.GetRackUnitSettings().TargetColor; });
+				LoudnessScaleParams.OnTargetColorChanged.BindLambda([=](FLinearColor NewColor)
+				{
+					SettingsHelper.GetRackUnitSettings().TargetColor = NewColor;
+					SettingsHelper.SaveConfig();
+				});
+				LoudnessScaleParams.OnTargetColorPickerWindowClosed.BindLambda([=](const TSharedRef<SWindow>& Window) { SettingsHelper.SaveConfig(); });
+				LoudnessScaleParams.OnResetTargetColorRequested.BindLambda([=]()
+				{
+					SettingsHelper.GetRackUnitSettings().TargetColor = FLoudnessMeterWidgetView::GetDefaultTargetColor();
+					SettingsHelper.SaveConfig();
+				});
+
 				LoudnessScaleParams.bRelativeScale.BindLambda([=]() { return SettingsHelper.GetRackUnitSettings().bUseRelativeLoudnessScale; });
 				LoudnessScaleParams.OnRelativeScaleToggleRequested.BindLambda([=]()
 					{
diff --git a/Engine/Plugins/Runtime/AudioWidgets/Source/AudioWidgetsCore/Private/LoudnessMeterWidgetView.cpp b/Engine/Plugins/Runtime/AudioWidgets/Source/AudioWidgetsCore/Private/LoudnessMeterWidgetView.cpp
--- a/Engine/Plugins/Runtime/AudioWidgets/Source/AudioWidgetsCore/Private/LoudnessMeterWidgetView.cpp
+++ b/Engine/Plugins/Runtime/AudioWidgets/Source/AudioWidgetsCore/Private/LoudnessMeterWidgetView.cpp
@@ -45,12 +45,14 @@ namespace AudioWidgetsCore
 		public:
 			SLATE_BEGIN_ARGS(SVariableDisplayRangeMeterWidget)
 				: _MeterValueColor(FLinearColor::Green)
+				, _TargetValueColor(FLinearColor::Transparent)
 				, _DisplayRange(60)
 				, _Offset(0)
 				{
 				}
 
 				SLATE_ATTRIBUTE(FSlateColor, MeterValueColor)
+				SLATE_ATTRIBUTE(FSlateColor, TargetValueColor)
 				SLATE_ATTRIBUTE(TArray<FAudioMeterChannelInfo>, MeterChannelInfo)
 				SLATE_ATTRIBUTE(int32, DisplayRange)
 				SLATE_ATTRIBUTE(int32, Offset)
@@ -59,9 +61,10 @@ namespace AudioWidgetsCore
 
 			void Construct(const SVariableDisplayRangeMeterWidget::FArguments& InArgs)
 			{
-				AudioMeterWidgetStyle->SetMeterSize({ 700.0, 25.0 }); // MaxLength, Width.
-				AudioMeterWidgetStyle->SetScaleHashHeight(4.0f); // This is the tick mark length.
-
+				AudioMeterWidgetStyle->SetMeterSize({ 700.0, 25.0 });	// MaxLength, Width.
+				AudioMeterWidgetStyle->SetScaleHashHeight(4.0f);		// This is the tick mark length.
+				AudioMeterWidgetStyle->SetPeakValueWidth(1.0f);			// The width used for displaying the loudness target value 
+				
 				DisplayRange = InArgs._DisplayRange;
 				Offset = InArgs._Offset;
 
@@ -70,7 +73,7 @@ namespace AudioWidgetsCore
 						.Orientation(EOrientation::Orient_Vertical)
 						.BackgroundColor(FLinearColor::Transparent)
 						.MeterBackgroundColor(FAudioMeterDefaultColorWidgetStyle::GetDefault().MeterBackgroundColor)
-						.MeterPeakColor(FLinearColor::Transparent)
+						.MeterPeakColor(InArgs._TargetValueColor)
 						.MeterClippingColor(FLinearColor::Transparent)
 						.MeterScaleColor(FLinearColor::White.CopyWithNewOpacity(0.25f))
 						.MeterScaleLabelColor(FLinearColor::White.CopyWithNewOpacity(0.25f))
@@ -287,7 +290,27 @@ namespace AudioWidgetsCore
 				[
 					SNew(SVariableDisplayRangeMeterWidget)
 						.MeterValueColor_Lambda([InItem]() { return InItem->Color.Get(); })
-						.MeterChannelInfo_Lambda([InItem, LoudnessValueConverter]() -> TArray<FAudioMeterChannelInfo>
+						.TargetValueColor_Lambda([InItem, LoudnessScaleParams]()
+						{
+							const bool bCanMeterShowTargetValue = InItem->MeterMetric.Get() == EAudioMeterMetric::Loudness && !InItem->IsRanged();
+							if (bCanMeterShowTargetValue)
+							{
+								const int32 Range  = LoudnessScaleParams.Range.Get();
+								const int32 Offset = LoudnessScaleParams.Offset.Get();
+								const int32 Target = LoudnessScaleParams.Target.Get();
+								
+								const int32 MinRangeValue = Offset - Range;
+								const int32 MaxRangeValue = Offset;
+
+								if (Target >= MinRangeValue && Target <= MaxRangeValue)
+								{
+									return LoudnessScaleParams.TargetColor.Get();
+								}
+							}
+
+							return FLinearColor::Transparent;
+						})
+						.MeterChannelInfo_Lambda([InItem, LoudnessValueConverter, Target = LoudnessScaleParams.Target, bRelativeScale = LoudnessScaleParams.bRelativeScale]() -> TArray<FAudioMeterChannelInfo>
 							{
 								constexpr float DefaultInvalidValue = MIN_VOLUME_DECIBELS;
 
@@ -323,7 +346,7 @@ namespace AudioWidgetsCore
 								{
 									{
 										.MeterValue = ConvertedValue.Get(DefaultInvalidValue),
-										.PeakValue = DefaultInvalidValue
+										.PeakValue  = bRelativeScale.Get() ? 0.0f : static_cast<float>(Target.Get())
 									}
 								};
 							})
@@ -603,6 +626,11 @@ namespace AudioWidgetsCore
 		return FStyleColors::AccentGreen.GetSpecifiedColor().CopyWithNewOpacity(0.5f);
 	}
 
+	FLinearColor FLoudnessMeterWidgetView::GetDefaultTargetColor()
+	{
+		return FStyleColors::AccentOrange.GetSpecifiedColor();
+	}
+
 	FLoudnessMeterWidgetView::FLoudnessMeterWidgetView()
 		: LoudnessMetrics(MakeShared<FLoudnessMetricRefArray>())
 		, ValuesCollectionView(MakeShared<FLoudnessMetricCollectionView>(LoudnessMetrics, [](const FLoudnessMetricRef& LoudnessMetricRef) { return LoudnessMetricRef->bShowValue.Get(); }))
@@ -701,6 +729,40 @@ namespace AudioWidgetsCore
 						MenuBuilder.AddWidget(TargetWidget, LOCTEXT("LoudnessScale_Target", "Target"));
 					}
 
+					if (LoudnessScaleParams.OnTargetColorChanged.IsBound())
+					{
+						MenuBuilder.AddMenuEntry(
+							LOCTEXT("LoudnessScale_TargetColor", "Target Color"),
+							FText(),
+							FSlateIcon(FAudioWidgetsCoreStyle::StyleName, "AudioWidgetsCoreStyle.ColorMixing"),
+							FExecuteAction::CreateLambda([
+								ParentWidget = FSlateApplication::Get().GetMenuHostWidget().ToWeakPtr(),
+								Color = LoudnessScaleParams.TargetColor,
+								OnColorChanged = LoudnessScaleParams.OnTargetColorChanged,
+								OnColorPickerWindowClosed = LoudnessScaleParams.OnTargetColorPickerWindowClosed]()
+								{
+									FColorPickerArgs PickerArgs(Color.Get(), OnColorChanged);
+									PickerArgs.ParentWidget = ParentWidget.Pin();
+									PickerArgs.bUseAlpha = true;
+									PickerArgs.bIsModal = true;
+									PickerArgs.OnColorPickerWindowClosed = OnColorPickerWindowClosed;
+									OpenColorPicker(PickerArgs);
+								}),
+							NAME_None,
+							EUserInterfaceActionType::Button);
+					}
+
+					if (LoudnessScaleParams.OnResetTargetColorRequested.IsBound())
+					{
+						MenuBuilder.AddMenuEntry(
+							LOCTEXT("LoudnessScale_ResetTargetColor", "Reset Target Color"),
+							FText(),
+							FSlateIcon(FAudioWidgetsCoreStyle::StyleName, "AudioWidgetsCoreStyle.Reset"),
+							LoudnessScaleParams.OnResetTargetColorRequested,
+							NAME_None,
+							EUserInterfaceActionType::Button);
+					}
+
 					if (LoudnessScaleParams.OnRelativeScaleToggleRequested.IsBound())
 					{
 						MenuBuilder.AddMenuEntry(
@@ -872,6 +934,7 @@ namespace AudioWidgetsCore
 			.BorderImage(&BackgroundBrush)
 			[
 				SNew(SVerticalBox)
+					.Clipping(EWidgetClipping::ClipToBounds)
 					+ SVerticalBox::Slot()
 					.AutoHeight()
 					[
diff --git a/Engine/Plugins/Runtime/AudioWidgets/Source/AudioWidgetsCore/Public/LoudnessMeterSettings.h b/Engine/Plugins/Runtime/AudioWidgets/Source/AudioWidgetsCore/Public/LoudnessMeterSettings.h
--- a/Engine/Plugins/Runtime/AudioWidgets/Source/AudioWidgetsCore/Public/LoudnessMeterSettings.h
+++ b/Engine/Plugins/Runtime/AudioWidgets/Source/AudioWidgetsCore/Public/LoudnessMeterSettings.h
@@ -35,6 +35,9 @@ struct FLoudnessMeterSettings
 	UPROPERTY(EditAnywhere, config, Category = LoudnessMeter, meta = (ClampMin = "-60", ClampMax = "30", UIMin = "-60", UIMax = "30"))
 	int LoudnessScaleTarget = -23;
 
+	UPROPERTY(EditAnywhere, config, Category = LoudnessMeter)
+	FLinearColor TargetColor = FStyleColors::AccentOrange.GetSpecifiedColor();
+
 	UPROPERTY(EditAnywhere, config, Category = LoudnessMeter)
 	bool bUseRelativeLoudnessScale = false;
 
diff --git a/Engine/Plugins/Runtime/AudioWidgets/Source/AudioWidgetsCore/Public/LoudnessMeterWidgetView.h b/Engine/Plugins/Runtime/AudioWidgets/Source/AudioWidgetsCore/Public/LoudnessMeterWidgetView.h
--- a/Engine/Plugins/Runtime/AudioWidgets/Source/AudioWidgetsCore/Public/LoudnessMeterWidgetView.h
+++ b/Engine/Plugins/Runtime/AudioWidgets/Source/AudioWidgetsCore/Public/LoudnessMeterWidgetView.h
@@ -24,6 +24,7 @@ namespace AudioWidgetsCore
 	{
 	public:
 		UE_API static FLinearColor GetDefaultMeterColor();
+		UE_API static FLinearColor GetDefaultTargetColor();
 
 		struct FLoudnessMetric
 		{
@@ -70,6 +71,11 @@ namespace AudioWidgetsCore
 			FOnIntValueChanged OnTargetValueChanged;
 			FOnIntValueCommitted OnTargetValueCommitted;
 
+			TAttribute<FLinearColor> TargetColor = GetDefaultTargetColor();
+			FOnLinearColorValueChanged OnTargetColorChanged;
+			FOnWindowClosed OnTargetColorPickerWindowClosed;
+			FSimpleDelegate OnResetTargetColorRequested;
+
 			TAttribute<bool> bRelativeScale = false;
 			FSimpleDelegate OnRelativeScaleToggleRequested;
 		};
