[AudioWidgets] [AudioMixer] [SignalProcessing] TruePeak: fine tuning so resulting values are more accurate, especially in FAudioLoudnessMeter since we were not taking the max among all channels.

#Synced-CL 51662785

diff --git a/Engine/Plugins/Runtime/AudioWidgets/Source/AudioWidgets/Private/AudioLoudnessMeter.cpp b/Engine/Plugins/Runtime/AudioWidgets/Source/AudioWidgets/Private/AudioLoudnessMeter.cpp
--- a/Engine/Plugins/Runtime/AudioWidgets/Source/AudioWidgets/Private/AudioLoudnessMeter.cpp
+++ b/Engine/Plugins/Runtime/AudioWidgets/Source/AudioWidgets/Private/AudioLoudnessMeter.cpp
@@ -107,6 +107,7 @@ namespace AudioWidgets
 		TruePeakSettings->PeakMode = EMeterPeakType::TruePeak;
 		TruePeakSettings->bIsAnalog = false;
 		TruePeakSettings->MeterAttackTime = 0;
+		TruePeakSettings->MeterReleaseTime = 0;
 
 		AudioDeviceId = InAudioDeviceId;
 
@@ -144,7 +145,7 @@ namespace AudioWidgets
 			}
 		}
 
-		LatestOverallTruePeakMeterResultsDelegateHandle.Reset();
+		PerChannelMeterResultsNativeDelegateHandle.Reset();
 		TruePeakAnalyzer.Reset();
 
 		LatestOverallLKFSResultsDelegateHandle.Reset();
@@ -202,21 +203,21 @@ namespace AudioWidgets
 	void FAudioLoudnessMeter::CreateTruePeakAnalyzer()
 	{
 		ensure(!TruePeakAnalyzer.IsValid());
-		ensure(!LatestOverallTruePeakMeterResultsDelegateHandle.IsValid());
+		ensure(!PerChannelMeterResultsNativeDelegateHandle.IsValid());
 
 		TruePeakAnalyzer = TStrongObjectPtr(NewObject<UMeterAnalyzer>());
 		TruePeakAnalyzer->Settings = TruePeakSettings.Get();
-		LatestOverallTruePeakMeterResultsDelegateHandle = TruePeakAnalyzer->OnLatestOverallMeterResultsNative.AddRaw(this, &FAudioLoudnessMeter::OnTruePeakOutput);
+		PerChannelMeterResultsNativeDelegateHandle = TruePeakAnalyzer->OnPerChannelMeterResultsNative.AddRaw(this, &FAudioLoudnessMeter::OnTruePeakOutput);
 	}
 
 	void FAudioLoudnessMeter::ReleaseTruePeakAnalyzer()
 	{
-		if (ensure(TruePeakAnalyzer.IsValid() && LatestOverallTruePeakMeterResultsDelegateHandle.IsValid()))
+		if (ensure(TruePeakAnalyzer.IsValid() && PerChannelMeterResultsNativeDelegateHandle.IsValid()))
 		{
-			TruePeakAnalyzer->OnLatestOverallMeterResultsNative.Remove(LatestOverallTruePeakMeterResultsDelegateHandle);
+			TruePeakAnalyzer->OnPerChannelMeterResultsNative.Remove(PerChannelMeterResultsNativeDelegateHandle);
 		}
 
-		LatestOverallTruePeakMeterResultsDelegateHandle.Reset();
+		PerChannelMeterResultsNativeDelegateHandle.Reset();
 		TruePeakAnalyzer.Reset();
 	}
 
@@ -528,14 +529,29 @@ namespace AudioWidgets
 		}
 	}
 
-	void FAudioLoudnessMeter::OnTruePeakOutput(UMeterAnalyzer* InMeterAnalyzer, const FMeterResults& InMeterResults)
+	void FAudioLoudnessMeter::OnTruePeakOutput(UMeterAnalyzer* InMeterAnalyzer, int32 InChannelIndex, const TArray<FMeterResults>& InMeterResultsArray)
 	{
-		if (InMeterAnalyzer == TruePeakAnalyzer.Get())
+		if (InMeterAnalyzer != TruePeakAnalyzer.Get())
 		{
-			LatestTruePeakResults = InMeterResults;
-			
-			LoudnessMeterMaxResults.TruePeakDb = FMath::Max<float>(LoudnessMeterMaxResults.TruePeakDb, InMeterResults.MeterValue);
+			return;
 		}
+
+		const FMeterResults* MaxResult = Algo::MaxElementBy(InMeterResultsArray, &FMeterResults::MeterValue);
+		if (!MaxResult)
+		{
+			return;
+		}
+
+		if (InChannelIndex == 0)
+		{
+			LatestTruePeakResults = *MaxResult;
+		}
+		else if (MaxResult->MeterValue > LatestTruePeakResults->MeterValue)
+		{
+			LatestTruePeakResults = *MaxResult;
+		}
+
+		LoudnessMeterMaxResults.TruePeakDb = FMath::Max(LoudnessMeterMaxResults.TruePeakDb, MaxResult->MeterValue);
 	}
 
 	FTimespan FAudioLoudnessMeter::GetAnalysisTime() const
diff --git a/Engine/Plugins/Runtime/AudioWidgets/Source/AudioWidgets/Private/AudioLoudnessMeter.h b/Engine/Plugins/Runtime/AudioWidgets/Source/AudioWidgets/Private/AudioLoudnessMeter.h
--- a/Engine/Plugins/Runtime/AudioWidgets/Source/AudioWidgets/Private/AudioLoudnessMeter.h
+++ b/Engine/Plugins/Runtime/AudioWidgets/Source/AudioWidgets/Private/AudioLoudnessMeter.h
@@ -60,7 +60,7 @@ namespace AudioWidgets
 		void InitLoudnessMeterWidgetView(const FAudioAnalyzerRackUnitConstructParams& Params);
 
 		void OnLoudnessOutput(ULKFSAnalyzer* InLKFSAnalyzer, const FLKFSResults& InLoudnessResults);
-		void OnTruePeakOutput(UMeterAnalyzer* InMeterAnalyzer, const FMeterResults& InMeterResults);
+		void OnTruePeakOutput(UMeterAnalyzer* InMeterAnalyzer, int32 InChannelIndex, const TArray<FMeterResults>& InMeterResultsArray);
 
 		FTimespan GetAnalysisTime() const;
 		FReply HandleResetButtonClicked();
@@ -87,7 +87,7 @@ namespace AudioWidgets
 		TStrongObjectPtr<UMeterSettings> TruePeakSettings;
 
 		/** Handle for results delegate for true peak analyzer. */
-		FDelegateHandle LatestOverallTruePeakMeterResultsDelegateHandle;
+		FDelegateHandle PerChannelMeterResultsNativeDelegateHandle;
 
 		/** Most recent results from the true peak analyzer. */
 		TOptional<FMeterResults> LatestTruePeakResults;
diff --git a/Engine/Source/Runtime/AudioMixer/Private/AudioMixerDevice.cpp b/Engine/Source/Runtime/AudioMixer/Private/AudioMixerDevice.cpp
--- a/Engine/Source/Runtime/AudioMixer/Private/AudioMixerDevice.cpp
+++ b/Engine/Source/Runtime/AudioMixer/Private/AudioMixerDevice.cpp
@@ -492,6 +492,7 @@ namespace Audio
 				.SampleRate = static_cast<float>(InSampleRate),
 				.NumChannels = InNumChannels,
 				.AttackTimeMsec = 0.0f,
+				.ReleaseTimeMsec = 0.0f,
 				.Mode = EPeakMode::TruePeak,
 				.bIsAnalog = false,
 			});
diff --git a/Engine/Source/Runtime/SignalProcessing/Private/EnvelopeFollower.cpp b/Engine/Source/Runtime/SignalProcessing/Private/EnvelopeFollower.cpp
--- a/Engine/Source/Runtime/SignalProcessing/Private/EnvelopeFollower.cpp
+++ b/Engine/Source/Runtime/SignalProcessing/Private/EnvelopeFollower.cpp
@@ -410,7 +410,7 @@ namespace Audio
 				}
 
 				const float MaxAbsValue = VectorReduceMax(VectorAbs(Accumulators));
-				OutBuffer[SampleIndex] = MaxAbsValue;
+				OutBuffer[SampleIndex] = FMath::Max(FMath::Abs(InputSample), MaxAbsValue);
 
 				const uint32 HistoryWriteIndex = (PrevHistoryWriteIndex + 1) & HistoryIndexMask;
 				ChannelSignalHistory[HistoryWriteIndex] = InputSample;
