[AudioWidgets] [AudioMixer] [Audio Insights] Loudness metering: set IntegratedLoudnessDuration to 30 mins as requested + changed "Long Term" labels to "Integrated" since sound designers are more familiarized with that term.

#Synced-CL 51662785

diff --git a/Engine/Plugins/AudioInsights/Source/AudioInsights/Private/Views/OutputMeterDashboardViewFactory.cpp b/Engine/Plugins/AudioInsights/Source/AudioInsights/Private/Views/OutputMeterDashboardViewFactory.cpp
--- a/Engine/Plugins/AudioInsights/Source/AudioInsights/Private/Views/OutputMeterDashboardViewFactory.cpp
+++ b/Engine/Plugins/AudioInsights/Source/AudioInsights/Private/Views/OutputMeterDashboardViewFactory.cpp
@@ -120,7 +120,7 @@ namespace UE::Audio::Insights
 		{
 			{
 				.Name = GET_MEMBER_NAME_CHECKED(FLoudnessMeterSettings, LongTermLoudness),
-				.DisplayName = LOCTEXT("LongTermLoudnessDisplayName", "Long Term"),
+				.DisplayName = LOCTEXT("IntegratedLoudnessDisplayName", "Integrated"),
 				.MeterMetric = EAudioMeterMetric::Loudness,
 				.Value = TAttribute<TOptional<float>>::CreateSPLambda(this, [this]() { return LoudnessMeterResults.LongTermLoudness; }),
 				.bShowValue = true,
diff --git a/Engine/Plugins/Runtime/AudioWidgets/Source/AudioWidgets/Private/AudioLoudnessMeter.cpp b/Engine/Plugins/Runtime/AudioWidgets/Source/AudioWidgets/Private/AudioLoudnessMeter.cpp
--- a/Engine/Plugins/Runtime/AudioWidgets/Source/AudioWidgets/Private/AudioLoudnessMeter.cpp
+++ b/Engine/Plugins/Runtime/AudioWidgets/Source/AudioWidgets/Private/AudioLoudnessMeter.cpp
@@ -101,6 +101,7 @@ namespace AudioWidgets
 		Teardown();
 
 		LKFSSettings = TStrongObjectPtr(NewObject<ULKFSSettings>());
+		LKFSSettings->IntegratedLoudnessDuration = 1800.0f; // 30 minutes
 		LKFSSettings->bCalculateOverallLoudnessRange = true;
 
 		TruePeakSettings = TStrongObjectPtr(NewObject<UMeterSettings>());
@@ -263,7 +264,7 @@ namespace AudioWidgets
 		{
 			{
 				.Name = GET_MEMBER_NAME_CHECKED(FLoudnessMeterRackUnitSettings, LongTermLoudness),
-				.DisplayName = LOCTEXT("LongTermLoudnessDisplayName", "Long Term"),
+				.DisplayName = LOCTEXT("IntegratedLoudnessDisplayName", "Integrated"),
 				.MeterMetric = EAudioMeterMetric::Loudness,
 				.Value = TAttribute<TOptional<float>>::CreateSPLambda(this, [this]() { return (LatestLoudnessResults.IsSet()) ? TOptional<float>(LatestLoudnessResults->GatedLoudness) : NullOpt; }),
 				.bShowValue = true,
diff --git a/Engine/Source/Runtime/AudioMixer/Private/AudioMixerDevice.cpp b/Engine/Source/Runtime/AudioMixer/Private/AudioMixerDevice.cpp
--- a/Engine/Source/Runtime/AudioMixer/Private/AudioMixerDevice.cpp
+++ b/Engine/Source/Runtime/AudioMixer/Private/AudioMixerDevice.cpp
@@ -503,6 +503,7 @@ namespace Audio
 	void FMixerDevice::InitializeMainSubmixLKFSAnalyzer(int32 InSampleRate, int32 InNumChannels)
 	{
 		FLKFSAnalyzerSettings LKFSAnalyzerSettings;
+		LKFSAnalyzerSettings.IntegratedLoudnessDuration = 1800.0f; // 30 minutes
 		LKFSAnalyzerSettings.bCalculateOverallLoudnessRange = true;
 
 		MainSubmixLKFSAnalyzer = MakeUnique<FLKFSAnalyzer>(InNumChannels, static_cast<float>(InSampleRate), LKFSAnalyzerSettings);
