[Audio Insights] Loudness meter: notify with a color change when a loudness value is greater than the set target loudness (or if true peak is greater than -1.0 dB as defined in the EBU128 standard)

#Synced-CL 51662785

diff --git a/Engine/Plugins/Runtime/AudioWidgets/Source/AudioWidgetsCore/Private/LoudnessMeterWidgetView.cpp b/Engine/Plugins/Runtime/AudioWidgets/Source/AudioWidgetsCore/Private/LoudnessMeterWidgetView.cpp
--- a/Engine/Plugins/Runtime/AudioWidgets/Source/AudioWidgetsCore/Private/LoudnessMeterWidgetView.cpp
+++ b/Engine/Plugins/Runtime/AudioWidgets/Source/AudioWidgetsCore/Private/LoudnessMeterWidgetView.cpp
@@ -38,6 +38,8 @@ namespace AudioWidgetsCore
 
 		const FText NegativeInfinityText = LOCTEXT("NegativeInfinity", "-∞");
 
+		constexpr float EBU128MaxPermittedTruePeakLevelDb = -1.0f;
+
 		bool SupportsMaxValueDisplay(const FName& InLoudnessMeterName)
 		{
 			return InLoudnessMeterName == GET_MEMBER_NAME_CHECKED(FLoudnessMeterSettings, ShortTermLoudness) ||
@@ -200,7 +202,46 @@ namespace AudioWidgetsCore
 		{
 			const FLoudnessValueConverter LoudnessValueConverter(LoudnessScaleParams);
 
-			const auto GetValueColorAndOpacity = [InItem]() { return InItem->Color.Get().CopyWithNewOpacity(1.0f); };
+			const auto GetValueColorAndOpacity = [InItem, LoudnessScaleParams, bInHoldMax]()
+			{
+				bool bIsValueAboveThreshold = false;
+
+				const TOptional<float>& ValueOpt = bInHoldMax.Get() ? InItem->MaxValue.Get() : InItem->Value.Get();
+
+				if (ValueOpt.IsSet())
+				{
+					const float Value = ValueOpt.GetValue();
+
+					switch (InItem->MeterMetric.Get())
+					{
+						case EAudioMeterMetric::Decibels:
+						{
+							if (Value > EBU128MaxPermittedTruePeakLevelDb)
+							{
+								bIsValueAboveThreshold = true;
+							}
+
+							break;
+						}
+
+						case EAudioMeterMetric::Loudness:
+						{
+							const int32 Target = LoudnessScaleParams.Target.Get();
+							if (!InItem->IsRanged() && Value > Target)
+							{
+								bIsValueAboveThreshold = true;
+							}
+
+							break;
+						}
+
+						default:
+							break;
+					}
+				}
+
+				return bIsValueAboveThreshold ? FLinearColor::Red : InItem->Color.Get().CopyWithNewOpacity(1.0f);
+			};
 
 			return SNew(SOverlay)
 				+ SOverlay::Slot()
