[Audio Insights] Loudness meter: added TruePeakLimit slider so we can customize the threshold in the same way as we do for Target loudness.

#Synced-CL 52248429

diff --git a/Engine/Plugins/AudioInsights/Source/AudioInsights/Private/Views/OutputMeterDashboardViewFactory.cpp b/Engine/Plugins/AudioInsights/Source/AudioInsights/Private/Views/OutputMeterDashboardViewFactory.cpp
--- a/Engine/Plugins/AudioInsights/Source/AudioInsights/Private/Views/OutputMeterDashboardViewFactory.cpp
+++ b/Engine/Plugins/AudioInsights/Source/AudioInsights/Private/Views/OutputMeterDashboardViewFactory.cpp
@@ -87,6 +87,13 @@ namespace UE::Audio::Insights
 			LoudnessMeterSettings.LoudnessScaleTarget = Value;
 		});
 
+		LoudnessScaleParams.TruePeakLimit.BindLambda([this]() { return LoudnessMeterSettings.TruePeakLimit; });
+		LoudnessScaleParams.OnTruePeakLimitValueChanged.BindLambda([this](float Value) { LoudnessMeterSettings.TruePeakLimit = Value; });
+		LoudnessScaleParams.OnTruePeakLimitValueCommitted.BindLambda([this](float Value, ETextCommit::Type CommitType)
+		{
+			LoudnessMeterSettings.TruePeakLimit = Value;
+		});
+
 		LoudnessScaleParams.TargetColor.BindLambda([this]() { return LoudnessMeterSettings.TargetColor; });
 		LoudnessScaleParams.OnTargetColorChanged.BindLambda([this](FLinearColor NewColor) {	LoudnessMeterSettings.TargetColor = NewColor; });
 		LoudnessScaleParams.OnResetTargetColorRequested.BindLambda([this]() { LoudnessMeterSettings.TargetColor = FLoudnessMeterWidgetView::GetDefaultTargetColor(); });
diff --git a/Engine/Plugins/Runtime/AudioWidgets/Source/AudioWidgets/Private/AudioLoudnessMeter.cpp b/Engine/Plugins/Runtime/AudioWidgets/Source/AudioWidgets/Private/AudioLoudnessMeter.cpp
--- a/Engine/Plugins/Runtime/AudioWidgets/Source/AudioWidgets/Private/AudioLoudnessMeter.cpp
+++ b/Engine/Plugins/Runtime/AudioWidgets/Source/AudioWidgets/Private/AudioLoudnessMeter.cpp
@@ -355,6 +355,14 @@ namespace AudioWidgets
 						SettingsHelper.SaveConfig();
 					});
 
+				LoudnessScaleParams.TruePeakLimit.BindLambda([=]() { return SettingsHelper.GetRackUnitSettings().TruePeakLimit; });
+				LoudnessScaleParams.OnTruePeakLimitValueChanged.BindLambda([=](float Value) { SettingsHelper.GetRackUnitSettings().TruePeakLimit = Value; });
+				LoudnessScaleParams.OnTruePeakLimitValueCommitted.BindLambda([=](float Value, ETextCommit::Type CommitType)
+					{
+						SettingsHelper.GetRackUnitSettings().TruePeakLimit = Value;
+						SettingsHelper.SaveConfig();
+					});
+
 				LoudnessScaleParams.TargetColor.BindLambda([=]() { return SettingsHelper.GetRackUnitSettings().TargetColor; });
 				LoudnessScaleParams.OnTargetColorChanged.BindLambda([=](FLinearColor NewColor)
 				{
diff --git a/Engine/Plugins/Runtime/AudioWidgets/Source/AudioWidgetsCore/Private/LoudnessMeterWidgetView.cpp b/Engine/Plugins/Runtime/AudioWidgets/Source/AudioWidgetsCore/Private/LoudnessMeterWidgetView.cpp
--- a/Engine/Plugins/Runtime/AudioWidgets/Source/AudioWidgetsCore/Private/LoudnessMeterWidgetView.cpp
+++ b/Engine/Plugins/Runtime/AudioWidgets/Source/AudioWidgetsCore/Private/LoudnessMeterWidgetView.cpp
@@ -38,8 +38,6 @@ namespace AudioWidgetsCore
 
 		const FText NegativeInfinityText = LOCTEXT("NegativeInfinity", "-∞");
 
-		constexpr float EBU128MaxPermittedTruePeakLevelDb = -1.0f;
-
 		bool SupportsMaxValueDisplay(const FName& InLoudnessMeterName)
 		{
 			return InLoudnessMeterName == GET_MEMBER_NAME_CHECKED(FLoudnessMeterSettings, ShortTermLoudness) ||
@@ -214,9 +212,11 @@ namespace AudioWidgetsCore
 
 					switch (InItem->MeterMetric.Get())
 					{
+						// True Peak meter
 						case EAudioMeterMetric::Decibels:
 						{
-							if (Value > EBU128MaxPermittedTruePeakLevelDb)
+							const float TruePeakLimit = LoudnessScaleParams.TruePeakLimit.Get();
+							if (Value > TruePeakLimit)
 							{
 								bIsValueAboveThreshold = true;
 							}
@@ -224,6 +224,7 @@ namespace AudioWidgetsCore
 							break;
 						}
 
+						// Loudness meters
 						case EAudioMeterMetric::Loudness:
 						{
 							const int32 Target = LoudnessScaleParams.Target.Get();
@@ -256,6 +257,7 @@ namespace AudioWidgetsCore
 							{
 								switch (InItem->MeterMetric.Get())
 								{
+									// True Peak meter
 									case EAudioMeterMetric::Decibels:
 									{
 										const TOptional<float>& ValueOpt = bInHoldMax.Get() ? InItem->MaxValue.Get() : InItem->Value.Get();
@@ -271,6 +273,7 @@ namespace AudioWidgetsCore
 											: FText::AsNumber(Value, &NumberFormattingOptions);
 									}
 
+									// Loudness meters
 									case EAudioMeterMetric::Loudness:
 									{
 										return InItem->IsRanged()
@@ -321,11 +324,11 @@ namespace AudioWidgetsCore
 		{
 			return SNew(SVerticalBox)
 				+ SVerticalBox::Slot()
-				.MinHeight(8)
-				.MaxHeight(8)
+				.MinHeight(2)
+				.MaxHeight(2)
 				[
 					SNew(SSpacer)
-						.Size(FVector2D(60.0f, 8.0f))
+					.Size(FVector2D(68.0f, 2.0f))
 				]
 				+ SVerticalBox::Slot()
 				.AutoHeight()
@@ -375,25 +378,50 @@ namespace AudioWidgetsCore
 						.MeterValueColor_Lambda([InItem]() { return InItem->Color.Get(); })
 						.TargetValueColor_Lambda([InItem, LoudnessScaleParams]()
 						{
-							const bool bCanMeterShowTargetValue = InItem->MeterMetric.Get() == EAudioMeterMetric::Loudness && !InItem->IsRanged();
-							if (bCanMeterShowTargetValue)
+							const int32 Range  = LoudnessScaleParams.Range.Get();
+							const int32 Offset = LoudnessScaleParams.Offset.Get();
+
+							const int32 MinRangeValue = Offset - Range;
+							const int32 MaxRangeValue = Offset;
+
+							switch (InItem->MeterMetric.Get())
 							{
-								const int32 Range  = LoudnessScaleParams.Range.Get();
-								const int32 Offset = LoudnessScaleParams.Offset.Get();
-								const int32 Target = LoudnessScaleParams.Target.Get();
-								
-								const int32 MinRangeValue = Offset - Range;
-								const int32 MaxRangeValue = Offset;
+								// True Peak meter
+								case EAudioMeterMetric::Decibels:
+								{
+									const float TruePeakLimit = LoudnessScaleParams.TruePeakLimit.Get();
 
-								if (Target >= MinRangeValue && Target <= MaxRangeValue)
+									if (TruePeakLimit >= MinRangeValue && TruePeakLimit <= MaxRangeValue)
+									{
+										return LoudnessScaleParams.TargetColor.Get();
+									}
+
+									break;
+								}
+
+								// Loudness meters
+								case EAudioMeterMetric::Loudness:
 								{
-									return LoudnessScaleParams.TargetColor.Get();
+									if (!InItem->IsRanged())
+									{
+										const int32 Target = LoudnessScaleParams.Target.Get();
+
+										if (Target >= MinRangeValue && Target <= MaxRangeValue)
+										{
+											return LoudnessScaleParams.TargetColor.Get();
+										}
+									}
+
+									break;
 								}
+
+								default:
+									break;
 							}
 
 							return FLinearColor::Transparent;
 						})
-						.MeterChannelInfo_Lambda([InItem, LoudnessValueConverter, Target = LoudnessScaleParams.Target, bRelativeScale = LoudnessScaleParams.bRelativeScale]() -> TArray<FAudioMeterChannelInfo>
+						.MeterChannelInfo_Lambda([InItem, LoudnessValueConverter, Target = LoudnessScaleParams.Target, bRelativeScale = LoudnessScaleParams.bRelativeScale, TruePeakLimit = LoudnessScaleParams.TruePeakLimit]() -> TArray<FAudioMeterChannelInfo>
 							{
 								constexpr float DefaultInvalidValue = MIN_VOLUME_DECIBELS;
 
@@ -404,7 +432,7 @@ namespace AudioWidgetsCore
 									{
 										{
 											.MeterValue = ValueOpt.Get(DefaultInvalidValue),
-											.PeakValue  = DefaultInvalidValue
+											.PeakValue  = TruePeakLimit.Get()
 										}
 									};
 								}
@@ -825,6 +853,23 @@ namespace AudioWidgetsCore
 						MenuBuilder.AddWidget(TargetWidget, LOCTEXT("LoudnessScale_Target", "Target"));
 					}
 
+					if (LoudnessScaleParams.OnTruePeakLimitValueCommitted.IsBound())
+					{
+						TSharedRef TruePeakLimitWidget = SNew(SNumericEntryBox<float>)
+							.MinDesiredValueWidth(100.0f)
+							.MinSliderValue(-60.0f)
+							.MaxSliderValue(30.0f)
+							.MinFractionalDigits(1)
+							.MaxFractionalDigits(1)
+							.Delta(0.5f)
+							.AllowSpin(true)
+							.Value_Lambda([TruePeakLimit = LoudnessScaleParams.TruePeakLimit]() { return TruePeakLimit.Get(); })
+							.OnValueChanged(LoudnessScaleParams.OnTruePeakLimitValueChanged)
+							.OnValueCommitted(LoudnessScaleParams.OnTruePeakLimitValueCommitted);
+
+						MenuBuilder.AddWidget(TruePeakLimitWidget, LOCTEXT("LoudnessScale_TruePeakLimit", "True Peak Limit"));
+					}
+
 					if (LoudnessScaleParams.OnTargetColorChanged.IsBound())
 					{
 						MenuBuilder.AddMenuEntry(
@@ -1097,6 +1142,83 @@ namespace AudioWidgetsCore
 					]
 					+ SVerticalBox::Slot()
 					.AutoHeight()
+					.Padding(8.0f, 2.0f)
+					[
+						// Target Loudness and True Peak Limit readout
+						SNew(SVerticalBox)
+						+ SVerticalBox::Slot()
+						.AutoHeight()
+						.Padding(0, 0, 0, 6)
+						.HAlign(HAlign_Center)
+						[
+							// Target Loudness readout
+							SNew(SHorizontalBox)
+							+ SHorizontalBox::Slot()
+							.AutoWidth()
+							[
+								SNew(STextBlock)
+								.Font(FCoreStyle::GetDefaultFontStyle("Bold", 10))
+								.Text(LOCTEXT("TargetLoudnessLabel", "  Target Loudness: "))
+							]
+							+ SHorizontalBox::Slot()
+							.AutoWidth()
+							[
+								SNew(STextBlock)
+								.Font(FStyleDefaults::GetFontInfo(10))
+								.Text_Lambda([this]()
+								{
+									return FText::Format(
+										LOCTEXT("TargetLoudnessValue", "{0} LUFS"),
+										FText::AsNumber(LoudnessScaleParams.Target.Get())
+									);
+								})
+							]
+						]
+						+ SVerticalBox::Slot()
+						.AutoHeight()
+						.HAlign(HAlign_Center)
+						[
+							// True Peak Limit readout
+							SNew(SHorizontalBox)
+							+ SHorizontalBox::Slot()
+							.AutoWidth()
+							[
+								SNew(STextBlock)
+								.Font(FCoreStyle::GetDefaultFontStyle("Bold", 10))
+								.Text(LOCTEXT("TruePeakLabel", "True Peak Limit: "))
+							]
+							+ SHorizontalBox::Slot()
+							.AutoWidth()
+							[
+								SNew(STextBlock)
+								.Font(FStyleDefaults::GetFontInfo(10))
+								.Text_Lambda([this]()
+								{
+									FNumberFormattingOptions FloatFormat;
+									FloatFormat.SetMinimumFractionalDigits(1);
+									FloatFormat.SetMaximumFractionalDigits(1);
+
+									return FText::Format(
+										LOCTEXT("TruePeakValue", "{0} dB"),
+										FText::AsNumber(LoudnessScaleParams.TruePeakLimit.Get(), &FloatFormat)
+									);
+								})
+							]
+						]
+					]
+					+ SVerticalBox::Slot()
+					.AutoHeight()
+					.Padding(8.0f, 2.0f)
+					[
+						SNew(SSeparator)
+							.Visibility_Lambda(GetTimerPanelVisibility)
+							.SeparatorImage(FAppStyle::Get().GetBrush("WhiteBrush"))
+							.Orientation(Orient_Horizontal)
+							.Thickness(1.0f)
+							.ColorAndOpacity(FLinearColor::White.CopyWithNewOpacity(0.25f))
+					]
+					+ SVerticalBox::Slot()
+					.AutoHeight()
 					[
 						SNew(SLoudnessMetricTileView,
 								FOnMakeTileInnerContent::CreateStatic(&MakeValueTileInnerContent, LoudnessScaleParams),
@@ -1105,7 +1227,7 @@ namespace AudioWidgetsCore
 							.Visibility(GetNumericValuesPanelVisibility)
 							.OnGenerateTile_Static(&MakeNumericValueTileWidget)
 							.ListItemsSource(ValuesCollectionView->GetVisibleItems())
-							.ItemHeight(82.0f)
+							.ItemHeight(75.0f)
 							.ItemWidth(68.0f)
 							.ItemAlignment(EListItemAlignment::CenterAligned)
 							.SelectionMode(ESelectionMode::Single)
diff --git a/Engine/Plugins/Runtime/AudioWidgets/Source/AudioWidgetsCore/Public/LoudnessMeterSettings.h b/Engine/Plugins/Runtime/AudioWidgets/Source/AudioWidgetsCore/Public/LoudnessMeterSettings.h
--- a/Engine/Plugins/Runtime/AudioWidgets/Source/AudioWidgetsCore/Public/LoudnessMeterSettings.h
+++ b/Engine/Plugins/Runtime/AudioWidgets/Source/AudioWidgetsCore/Public/LoudnessMeterSettings.h
@@ -41,6 +41,9 @@ struct FLoudnessMeterSettings
 	UPROPERTY(EditAnywhere, config, Category = LoudnessMeter, meta = (ClampMin = "-60", ClampMax = "30", UIMin = "-60", UIMax = "30"))
 	int LoudnessScaleTarget = -23;
 
+	UPROPERTY(EditAnywhere, config, Category = LoudnessMeter, meta = (ClampMin = "-60.0", ClampMax = "30.0", UIMin = "-60.0", UIMax = "30.0"))
+	float TruePeakLimit = -1.0f;
+
 	UPROPERTY(EditAnywhere, config, Category = LoudnessMeter)
 	FLinearColor TargetColor = FStyleColors::AccentOrange.GetSpecifiedColor();
 
diff --git a/Engine/Plugins/Runtime/AudioWidgets/Source/AudioWidgetsCore/Public/LoudnessMeterWidgetView.h b/Engine/Plugins/Runtime/AudioWidgets/Source/AudioWidgetsCore/Public/LoudnessMeterWidgetView.h
--- a/Engine/Plugins/Runtime/AudioWidgets/Source/AudioWidgetsCore/Public/LoudnessMeterWidgetView.h
+++ b/Engine/Plugins/Runtime/AudioWidgets/Source/AudioWidgetsCore/Public/LoudnessMeterWidgetView.h
@@ -64,6 +64,9 @@ namespace AudioWidgetsCore
 		DECLARE_DELEGATE_OneParam(FOnIntValueChanged, int32 /*NewValue*/);
 		DECLARE_DELEGATE_TwoParams(FOnIntValueCommitted, int32 /*NewValue*/, ETextCommit::Type /*CommitType*/);
 
+		DECLARE_DELEGATE_OneParam(FOnFloatValueChanged, float /*NewValue*/);
+		DECLARE_DELEGATE_TwoParams(FOnFloatValueCommitted, float /*NewValue*/, ETextCommit::Type /*CommitType*/);
+
 		struct FLoudnessScaleParams
 		{
 			TAttribute<int32> Range = 60;
@@ -78,6 +81,11 @@ namespace AudioWidgetsCore
 			FOnIntValueChanged OnTargetValueChanged;
 			FOnIntValueCommitted OnTargetValueCommitted;
 
+			// Using float here to have a more fine grained representation of true peak in dB where fractional precision matters (unlike LUFS values).
+			TAttribute<float> TruePeakLimit = -1.0f; 
+			FOnFloatValueChanged OnTruePeakLimitValueChanged;
+			FOnFloatValueCommitted OnTruePeakLimitValueCommitted;
+
 			TAttribute<FLinearColor> TargetColor = GetDefaultTargetColor();
 			FOnLinearColorValueChanged OnTargetColorChanged;
 			FOnWindowClosed OnTargetColorPickerWindowClosed;
