[AudioWidgets] [Audio Insights] AudioAnalyzerRack: added logic to exclude specific audio analyzer rack units + made AudioLoudnessMeter.h public so RackUnitTypeInfo can be accessed from Audio Insights.

#Synced-CL 52441547

diff --git a/Engine/Plugins/AudioInsights/Source/AudioInsightsEditor/Private/Analyzers/SubmixAudioAnalyzerRack.cpp b/Engine/Plugins/AudioInsights/Source/AudioInsightsEditor/Private/Analyzers/SubmixAudioAnalyzerRack.cpp
--- a/Engine/Plugins/AudioInsights/Source/AudioInsightsEditor/Private/Analyzers/SubmixAudioAnalyzerRack.cpp
+++ b/Engine/Plugins/AudioInsights/Source/AudioInsightsEditor/Private/Analyzers/SubmixAudioAnalyzerRack.cpp
@@ -6,6 +6,7 @@
 #include "AudioInsightsEditorModule.h"
 #include "AudioInsightsEditorSettings.h"
 #include "AudioInsightsStyle.h"
+#include "AudioLoudnessMeter.h"
 #include "AudioMaterialSlate/AudioMaterialSlateTypes.h"
 #include "AudioMeter.h"
 #include "AudioMixerDevice.h"
@@ -80,6 +81,7 @@ namespace UE::Audio::Insights
 				.TabManagerLayoutName = TEXT("AudioInsights_FAudioAnalyzerRack_v0"),
 				.StyleSet = &FAnalyzerRackStyleSet::Get(),
 				.EditorSettingsClass = UAudioInsightsEditorSettings::StaticClass(),
+				.ExcludedRackUnitTypes = { AudioWidgets::FAudioLoudnessMeter::RackUnitTypeInfo.TypeName, AudioWidgets::FAudioMeter::RackUnitTypeInfo.TypeName },
 			};
 
 			return MakeShared<FAudioAnalyzerRack>(Params);
diff --git a/Engine/Plugins/Runtime/AudioWidgets/Source/AudioWidgets/Private/AudioAnalyzerRack.cpp b/Engine/Plugins/Runtime/AudioWidgets/Source/AudioWidgets/Private/AudioAnalyzerRack.cpp
--- a/Engine/Plugins/Runtime/AudioWidgets/Source/AudioWidgets/Private/AudioAnalyzerRack.cpp
+++ b/Engine/Plugins/Runtime/AudioWidgets/Source/AudioWidgets/Private/AudioAnalyzerRack.cpp
@@ -41,6 +41,7 @@ namespace AudioWidgets
 			TabManagerLayoutName = FName(StringBuilder);
 		}
 
+		ExcludedRackUnitTypes = Params.ExcludedRackUnitTypes;
 		RackUnitConstructParams.StyleSet = (Params.StyleSet != nullptr) ? Params.StyleSet : &FAudioWidgetsStyle::Get();
 		RackUnitConstructParams.EditorSettingsClass = Params.EditorSettingsClass;
 	}
@@ -98,6 +99,12 @@ namespace AudioWidgets
 
 		TArray<const FAudioAnalyzerRackUnitTypeInfo*> RackUnitTypes;
 		FAudioAnalyzerRackUnitRegistry::Get().GetRegisteredRackUnitTypes(RackUnitTypes);
+		
+		if (!ExcludedRackUnitTypes.IsEmpty())
+		{
+			RackUnitTypes.RemoveAll([this](const FAudioAnalyzerRackUnitTypeInfo* Type) { return ExcludedRackUnitTypes.Contains(Type->TypeName); });
+		}
+
 		for (const FAudioAnalyzerRackUnitTypeInfo* RackUnitType : RackUnitTypes)
 		{
 			TabManager->RegisterTabSpawner(RackUnitType->TypeName, FOnSpawnTab::CreateLambda([this](const FSpawnTabArgs& Args)
diff --git a/Engine/Plugins/Runtime/AudioWidgets/Source/AudioWidgets/Private/AudioLoudnessMeter.h b/Engine/Plugins/Runtime/AudioWidgets/Source/AudioWidgets/Private/AudioLoudnessMeter.h
deleted file mode 100644
diff --git a/Engine/Plugins/Runtime/AudioWidgets/Source/AudioWidgets/Public/AudioAnalyzerRack.h b/Engine/Plugins/Runtime/AudioWidgets/Source/AudioWidgets/Public/AudioAnalyzerRack.h
--- a/Engine/Plugins/Runtime/AudioWidgets/Source/AudioWidgets/Public/AudioAnalyzerRack.h
+++ b/Engine/Plugins/Runtime/AudioWidgets/Source/AudioWidgets/Public/AudioAnalyzerRack.h
@@ -85,6 +85,9 @@ namespace AudioWidgets
 
 			/** An Editor Settings class can be provided for rack units that require one for saving settings. */
 			const UClass* EditorSettingsClass = nullptr;
+
+			/** Rack unit type names to exclude from the rack (they won't appear in the menu or layout). */
+			TSet<FName> ExcludedRackUnitTypes;
 		};
 
 		UE_API FAudioAnalyzerRack(const FRackConstructParams& Params);
@@ -113,6 +116,7 @@ namespace AudioWidgets
 		TSharedRef<IAudioAnalyzerRackUnit> MakeRackUnit(FName RackUnitTypeName);
 
 		FName TabManagerLayoutName;
+		TSet<FName> ExcludedRackUnitTypes;
 
 		TStrongObjectPtr<UAudioBus> AudioBus;
 		FAudioAnalyzerRackUnitConstructParams RackUnitConstructParams;
diff --git a/Engine/Plugins/Runtime/AudioWidgets/Source/AudioWidgets/Public/AudioLoudnessMeter.h b/Engine/Plugins/Runtime/AudioWidgets/Source/AudioWidgets/Public/AudioLoudnessMeter.h
new file mode 100644
--- /dev/null
+++ b/Engine/Plugins/Runtime/AudioWidgets/Source/AudioWidgets/Public/AudioLoudnessMeter.h
@@ -0,0 +1,110 @@
+// Copyright Epic Games, Inc. All Rights Reserved.
+#pragma once
+
+#include "AudioAnalyzerRack.h"
+#include "Containers/Ticker.h"
+#include "LKFS.h"
+#include "LoudnessMeterWidgetView.h"
+#include "Meter.h"
+
+#define UE_API AUDIOWIDGETS_API
+
+namespace AudioWidgets
+{
+	/**
+	 * Constructor parameters for the analyzer.
+	 */
+	struct FAudioLoudnessMeterParams
+	{
+		int32 NumChannels = 1;
+		Audio::FDeviceId AudioDeviceId = INDEX_NONE;
+		TObjectPtr<UAudioBus> ExternalAudioBus = nullptr;
+	};
+
+	/**
+	 * Owns an analyzer and a corresponding Slate widget for displaying the loudness stats.
+	 * Can either create an Audio Bus to analyze, or analyze the given Bus.
+	 */
+	class FAudioLoudnessMeter : public IAudioAnalyzerRackUnit
+	{
+	public:
+		static UE_API const FAudioAnalyzerRackUnitTypeInfo RackUnitTypeInfo;
+		
+		UE_API FAudioLoudnessMeter(const FAudioLoudnessMeterParams& Params);
+		UE_API ~FAudioLoudnessMeter();
+
+		// Begin IAudioAnalyzerRackUnit overrides.
+		UE_API virtual void SetAudioBusInfo(const FAudioBusInfo& AudioBusInfo) override;
+		UE_API virtual TSharedRef<SDockTab> SpawnTab(const FSpawnTabArgs& Args) const override;
+		// End IAudioAnalyzerRackUnit overrides.
+
+	private:
+		struct FLoudnessMeterMaxResults
+		{
+			float ShortTermLoudness = MIN_VOLUME_DECIBELS;
+			float MomentaryLoudness = MIN_VOLUME_DECIBELS;
+			float TruePeakDb        = MIN_VOLUME_DECIBELS;
+		};
+
+		static TSharedRef<IAudioAnalyzerRackUnit> MakeRackUnit(const FAudioAnalyzerRackUnitConstructParams& Params);
+
+		void Init(int32 InNumChannels, Audio::FDeviceId InAudioDeviceId, TObjectPtr<UAudioBus> InExternalAudioBus);
+		void Teardown();
+
+		void CreateLKFSAnalyzer();
+		void ReleaseLKFSAnalyzer();
+		void ResetLKFSAnalyzer();
+
+		void CreateTruePeakAnalyzer();
+		void ReleaseTruePeakAnalyzer();
+		void ResetTruePeakAnalyzer();
+
+		void InitLoudnessMeterWidgetView(const FAudioAnalyzerRackUnitConstructParams& Params);
+
+		void OnLoudnessOutput(ULKFSAnalyzer* InLKFSAnalyzer, const FLKFSResults& InLoudnessResults);
+		void OnTruePeakOutput(UMeterAnalyzer* InMeterAnalyzer, int32 InChannelIndex, const TArray<FMeterResults>& InMeterResultsArray);
+
+		FTimespan GetAnalysisTime() const;
+		FReply HandleResetButtonClicked();
+
+		/** Creates the meters widget */
+		AudioWidgetsCore::FLoudnessMeterWidgetView LoudnessMeterWidgetView;
+		
+		/** Loudness analyzer object. */
+		TStrongObjectPtr<ULKFSAnalyzer> LKFSAnalyzer;
+
+		/** Loudness analyzer settings. */
+		TStrongObjectPtr<ULKFSSettings> LKFSSettings;
+
+		/** Handle for results delegate for loudness analyzer. */
+		FDelegateHandle LatestOverallLKFSResultsDelegateHandle;
+
+		/** Most recent results from the loudness analyzer. */
+		TOptional<FLKFSResults> LatestLoudnessResults;
+
+		/** True peak analyzer object. */
+		TStrongObjectPtr<UMeterAnalyzer> TruePeakAnalyzer;
+
+		/** True peak analyzer settings. */
+		TStrongObjectPtr<UMeterSettings> TruePeakSettings;
+
+		/** Handle for results delegate for true peak analyzer. */
+		FDelegateHandle PerChannelMeterResultsNativeDelegateHandle;
+
+		/** Most recent results from the true peak analyzer. */
+		TOptional<FMeterResults> LatestTruePeakResults;
+
+		/** Loudness meter max results (only for short term loudness, momentary loudness and true peak). */
+		FLoudnessMeterMaxResults LoudnessMeterMaxResults;
+
+		Audio::FDeviceId AudioDeviceId = INDEX_NONE;
+
+		FTSTicker::FDelegateHandle WidgetRefreshTicker;
+
+		/** The audio bus used for analysis. */
+		TStrongObjectPtr<UAudioBus> AudioBus;
+		bool bUseExternalAudioBus = false;
+	};
+} // namespace AudioWidgets
+
+#undef UE_API
