[Waveform Editor] - Panning past isolated Loop Region is will rubberband graph to the start

Introduce panning flag on the transport coordinator to prevent SetProgressRatio from resetting the display range during pan. Also fixes and simplifies context menu positioning for Markers.

#Synced-CL 52479141

diff --git a/Engine/Plugins/Editor/WaveformEditor/Source/WaveformEditor/Private/WaveformEditor.cpp b/Engine/Plugins/Editor/WaveformEditor/Source/WaveformEditor/Private/WaveformEditor.cpp
--- a/Engine/Plugins/Editor/WaveformEditor/Source/WaveformEditor/Private/WaveformEditor.cpp
+++ b/Engine/Plugins/Editor/WaveformEditor/Source/WaveformEditor/Private/WaveformEditor.cpp
@@ -1894,7 +1894,7 @@ void FWaveformEditor::UpdateTransportState()
 		case EAudioComponentPlayState::Stopped:
 			// Set TransportController and displayed TransportCoordinator to value clamped within TransformationChainConfig.
 
-			if (!TransportCoordinator->IsScrubbing())
+			if (!TransportCoordinator->IsScrubbing() && !TransportCoordinator->IsScrolling())
 			{
 				check(SoundWave->GetImportedSampleRate() > 0);
 				const float PlayheadTime = static_cast<float>(WaveformView.ViewWidget->GetPlayheadFramePosition()) / static_cast<float>(SoundWave->GetImportedSampleRate());
@@ -2192,7 +2192,7 @@ void FWaveformEditor::HandleAudioComponentPlayStateChanged(const UAudioComponent
 	default:
 		break;
 	case EAudioComponentPlayState::Stopped:
-		if (!TransportCoordinator->IsScrubbing())
+		if (!TransportCoordinator->IsScrubbing() && !TransportCoordinator->IsScrolling())
 		{
 			// USoundWave::PostEditChangeProperty calls USoundWave::UpdateAsset which frees the soundwave's resources,
 			// stopping the sound, calling this handler.  Avoiding this implicit state change is currently too risky and
diff --git a/Engine/Plugins/Editor/WaveformEditor/Source/WaveformEditorWidgets/Private/WaveformEditorZoomController.cpp b/Engine/Plugins/Editor/WaveformEditor/Source/WaveformEditorWidgets/Private/WaveformEditorZoomController.cpp
--- a/Engine/Plugins/Editor/WaveformEditor/Source/WaveformEditorWidgets/Private/WaveformEditorZoomController.cpp
+++ b/Engine/Plugins/Editor/WaveformEditor/Source/WaveformEditorWidgets/Private/WaveformEditorZoomController.cpp
@@ -61,13 +61,16 @@ float FWaveformEditorZoomController::GetZoomRatio() const
 
 void FWaveformEditorZoomController::CheckBounds(const FGeometry& AllottedGeometry)
 {
-	if (bIsPanning)
+	if (bIsScrolling)
 	{
 		FVector2D MouseAbsolutePosition = UWidgetLayoutLibrary::GetMousePositionOnPlatform();
 
 		if (!AllottedGeometry.IsUnderLocation(MouseAbsolutePosition))
 		{
-			bIsPanning = false;
+			bIsScrolling = false;
+
+			check(TransportController);
+			TransportController->SetScrolling(false);
 		}
 	}
 }
@@ -79,7 +82,12 @@ FReply FWaveformEditorZoomController::OnMouseButtonDown(const FGeometry& MyGeome
 		const FVector2D LocalCursorPosition = MyGeometry.AbsoluteToLocal(MouseEvent.GetScreenSpacePosition());
 		LocalCursorXPosition = FMath::Clamp(LocalCursorPosition.X, 0.f, MyGeometry.GetLocalSize().X);
 
-		bIsPanning = true;
+		bIsScrolling = true;
+
+		check(TransportController);
+		TransportController->SetScrolling(true);
+
+		return FReply::Handled();
 	}
 
 	return FReply::Unhandled();
@@ -89,7 +97,12 @@ FReply FWaveformEditorZoomController::OnMouseButtonUp(const FGeometry& MyGeometr
 {
 	if (MouseEvent.GetEffectingButton() == EKeys::RightMouseButton)
 	{
-		bIsPanning = false;
+		bIsScrolling = false;
+
+		check(TransportController);
+		TransportController->SetScrolling(false);
+
+		return FReply::Handled();
 	}
 
 	return FReply::Unhandled();
@@ -97,7 +110,7 @@ FReply FWaveformEditorZoomController::OnMouseButtonUp(const FGeometry& MyGeometr
 
 FReply FWaveformEditorZoomController::OnMouseMove(const FGeometry& MyGeometry, const FPointerEvent& MouseEvent)
 {
-	if (bIsPanning)
+	if (bIsScrolling)
 	{
 		check(TransportController);
 		check(MyGeometry.GetLocalSize().X > 0);
diff --git a/Engine/Plugins/Editor/WaveformEditor/Source/WaveformEditorWidgets/Public/WaveformEditorZoomController.h b/Engine/Plugins/Editor/WaveformEditor/Source/WaveformEditorWidgets/Public/WaveformEditorZoomController.h
--- a/Engine/Plugins/Editor/WaveformEditor/Source/WaveformEditorWidgets/Public/WaveformEditorZoomController.h
+++ b/Engine/Plugins/Editor/WaveformEditor/Source/WaveformEditorWidgets/Public/WaveformEditorZoomController.h
@@ -47,7 +47,7 @@ private:
 	float ZoomLevelStep = 2.f;
 	const uint32 LogRatioBase = 100;
 
-	bool bIsPanning = false;
+	bool bIsScrolling = false;
 	float LocalCursorXPosition = 0.0f;
 
 	TSharedPtr<FSparseSampledSequenceTransportCoordinator> TransportController = nullptr;
diff --git a/Engine/Plugins/Editor/WaveformEditor/Source/WaveformTransformationsWidgets/Private/WaveformTransformationMarkersRenderer.cpp b/Engine/Plugins/Editor/WaveformEditor/Source/WaveformTransformationsWidgets/Private/WaveformTransformationMarkersRenderer.cpp
--- a/Engine/Plugins/Editor/WaveformEditor/Source/WaveformTransformationsWidgets/Private/WaveformTransformationMarkersRenderer.cpp
+++ b/Engine/Plugins/Editor/WaveformEditor/Source/WaveformTransformationsWidgets/Private/WaveformTransformationMarkersRenderer.cpp
@@ -377,6 +377,11 @@ FReply FWaveformTransformationMarkerRenderer::OnMouseButtonDown(SWidget& OwnerWi
 
 FReply FWaveformTransformationMarkerRenderer::OnMouseButtonUp(SWidget& OwnerWidget, const FGeometry& MyGeometry, const FPointerEvent& MouseEvent)
 {
+	if (MarkersArray == nullptr || MarkerInteractionType == EMarkerInteractionType::None)
+	{
+		return FReply::Unhandled();
+	}
+
 	check(StrongMarkersTransformation);
 	bool bHasTransformationChanged = false;
 
@@ -390,11 +395,6 @@ FReply FWaveformTransformationMarkerRenderer::OnMouseButtonUp(SWidget& OwnerWidg
 		bHasTransformationChanged = true;
 	}
 
-	if (MarkersArray == nullptr)
-	{
-		return FReply::Unhandled();
-	}
-
 	if (MarkerInteractionType != EMarkerInteractionType::None)
 	{
 		bool bMarkDirty = MarkerInteractionType == EMarkerInteractionType::ScrubbingMarkerHandleLeft || MarkerInteractionType == EMarkerInteractionType::ScrubbingMarkerHandleRight;
@@ -875,16 +875,10 @@ void FWaveformTransformationMarkerRenderer::ShowConvertToLoopAtMarker(const FGeo
 		FSlateApplication::Get().OnApplicationActivationStateChanged().Remove(ApplicationActivationStateHandle);
 	}
 
-	const float ContextMenuMaxSize = 750.0f;
-	const FVector2D ContexMenuMaxPosition(LocalWindowMaxPosition.X > ContextMenuMaxSize ? LocalWindowMaxPosition.X - ContextMenuMaxSize : 0.0f
-	                                     ,LocalWindowMaxPosition.Y > ContextMenuMaxSize ? LocalWindowMaxPosition.Y - ContextMenuMaxSize : 0.0f);
-	const FVector2D ContexMenuPosition( FMath::Min(LocalCursorPosition.X, ContexMenuMaxPosition.X),
-										FMath::Min(LocalCursorPosition.Y, ContexMenuMaxPosition.Y) );
-
 	TSharedRef<SWindow> MenuWindow = SNew(SWindow)
 		.AutoCenter(EAutoCenter::None)
 		.Type(EWindowType::Menu)
-		.ScreenPosition(ContexMenuPosition)
+		.ScreenPosition(LocalCursorPosition)
 		.SizingRule(ESizingRule::Autosized)
 		.SupportsTransparency(EWindowTransparency::None)
 		.IsPopupWindow(true)
diff --git a/Engine/Plugins/Runtime/AudioWidgets/Source/AudioWidgets/Private/SparseSampledSequenceTransportCoordinator.cpp b/Engine/Plugins/Runtime/AudioWidgets/Source/AudioWidgets/Private/SparseSampledSequenceTransportCoordinator.cpp
--- a/Engine/Plugins/Runtime/AudioWidgets/Source/AudioWidgets/Private/SparseSampledSequenceTransportCoordinator.cpp
+++ b/Engine/Plugins/Runtime/AudioWidgets/Source/AudioWidgets/Private/SparseSampledSequenceTransportCoordinator.cpp
@@ -35,6 +35,16 @@ const bool FSparseSampledSequenceTransportCoordinator::IsScrubbing() const
 	return bIsScrubbing;
 }
 
+const bool FSparseSampledSequenceTransportCoordinator::IsScrolling() const
+{
+	return bIsScrolling;
+}
+
+void FSparseSampledSequenceTransportCoordinator::SetScrolling(const bool bInIsScrolling)
+{
+	bIsScrolling = bInIsScrolling;
+}
+
 float FSparseSampledSequenceTransportCoordinator::ConvertAbsoluteRatioToZoomed(const float InAbsoluteRatio) const
 {
 	return (InAbsoluteRatio - DisplayRange.GetLowerBoundValue()) / ZoomRatio;
@@ -179,7 +189,7 @@ void FSparseSampledSequenceTransportCoordinator::ScrubFocusPoint(const float InT
 
 void FSparseSampledSequenceTransportCoordinator::SetProgressRatio(const float NewRatio)
 {	 
-	if (bIsScrubbing)
+	if (bIsScrubbing || bIsScrolling)
 	{
 		return;
 	}
@@ -191,34 +201,45 @@ void FSparseSampledSequenceTransportCoordinator::SetProgressRatio(const float Ne
 
 	const double NewPlayheadPosition = (ScaledPlaybackRatio - DisplayRange.GetLowerBoundValue()) / ZoomRatio;
 
-	MoveFocusPoint(NewPlayheadPosition);
-
-	if (FocusPoint < 0.f)
+	if (NewPlayheadPosition < 0.f)
 	{
-		//if sound is looping and we are zoomed, playhead position will be negative when new loop starts
-		//hence we reset the view to go back at the beginning of the file
+		// If the playhead is off-screen to the left (e.g. loop restarted or user scrolled away),
+		// recenter the display range with the playback position on the left edge
 
-		MoveFocusPoint(0.f);
-		UpdateDisplayRange(FocusPoint, ZoomRatio);
-	}
+		double NewDisplayLowBound = FMath::Clamp(ScaledPlaybackRatio, 0.0, 1.0 - ZoomRatio);
+		double NewDisplayUpBound = NewDisplayLowBound + ZoomRatio;
 
-	else if (FocusPoint >= FocusPointLockPosition)
+		// Update DisplayRange first so the Playhead is placed correctly at the left edge
+		UpdateDisplayRange(NewDisplayLowBound, NewDisplayUpBound);
+		
+		constexpr float LeftBound = 0.f;
+		MoveFocusPoint(LeftBound);
+	}
+	else if (NewPlayheadPosition >= FocusPointLockPosition)
 	{
-
 		const double NewDisplayUpBound = FMath::Clamp(ScaledPlaybackRatio + (ZoomRatio * (1.f - FocusPointLockPosition)), 0.f, 1.f);
 		double NewDisplayLowBound = NewDisplayUpBound < 1.f ? ScaledPlaybackRatio - (ZoomRatio * FocusPointLockPosition) : 1.f - ZoomRatio;
 
 		NewDisplayLowBound = FMath::Clamp(NewDisplayLowBound, 0, 1);
+		
+		UpdateDisplayRange(NewDisplayLowBound, NewDisplayUpBound);
 
 		if (NewDisplayUpBound < 1.f)
 		{
-			//we only lock the playhead if the display upper bound is before the end of the waveform
-			//otherwise we let the playhead go to the end
+			// We only lock the playhead if the display upper bound is before the end of the waveform
+			// otherwise we let the playhead go to the end
 
 			MoveFocusPoint(FocusPointLockPosition);
 		}
-
-		UpdateDisplayRange(NewDisplayLowBound, NewDisplayUpBound);
+		else
+		{
+			// If playhead is far off screen, the NewPlayheadPosition could be > 1.0
+			MoveFocusPoint(NewPlayheadPosition);
+		}
+	}
+	else
+	{
+		MoveFocusPoint(NewPlayheadPosition);
 	}
 }
 
diff --git a/Engine/Plugins/Runtime/AudioWidgets/Source/AudioWidgets/Public/SparseSampledSequenceTransportCoordinator.h b/Engine/Plugins/Runtime/AudioWidgets/Source/AudioWidgets/Public/SparseSampledSequenceTransportCoordinator.h
--- a/Engine/Plugins/Runtime/AudioWidgets/Source/AudioWidgets/Public/SparseSampledSequenceTransportCoordinator.h
+++ b/Engine/Plugins/Runtime/AudioWidgets/Source/AudioWidgets/Public/SparseSampledSequenceTransportCoordinator.h
@@ -31,6 +31,8 @@ public:
 	UE_API const double GetFocusPoint() const override;
 	UE_API void ScrubFocusPoint(const float InTargetFocusPoint, const bool bIsMoving) override;
 	UE_API const bool IsScrubbing() const;
+	UE_API const bool IsScrolling() const;
+	UE_API void SetScrolling(const bool bInIsScrolling);
 	UE_API void SetProgressRatio(const float NewRatio) override;
 	UE_API void SetZoomRatio(const float NewRatio) override;
 	UE_API float ConvertAbsoluteRatioToZoomed(const float InAbsoluteRatio) const override;
@@ -59,6 +61,7 @@ private:
 	TRange<double> ProgressRange = TRange<double>::Inclusive(0, 1);
 
 	bool bIsScrubbing = false;
+	bool bIsScrolling = false;
 };
 
 #undef UE_API
