[Waveform editor] - Crash in STransformedWaveformViewPanel::SetPlayheadRatio() when zoomed in and playing back a sound

Fix: clamp NewPlayheadPosition to [0.0, 1.0] before passing it to MoveFocusPoint.

#Synced-CL 52479141

diff --git a/Engine/Plugins/Runtime/AudioWidgets/Source/AudioWidgets/Private/SparseSampledSequenceTransportCoordinator.cpp b/Engine/Plugins/Runtime/AudioWidgets/Source/AudioWidgets/Private/SparseSampledSequenceTransportCoordinator.cpp
--- a/Engine/Plugins/Runtime/AudioWidgets/Source/AudioWidgets/Private/SparseSampledSequenceTransportCoordinator.cpp
+++ b/Engine/Plugins/Runtime/AudioWidgets/Source/AudioWidgets/Private/SparseSampledSequenceTransportCoordinator.cpp
@@ -212,7 +212,7 @@ void FSparseSampledSequenceTransportCoordinator::SetProgressRatio(const float Ne
 		// Update DisplayRange first so the Playhead is placed correctly at the left edge
 		UpdateDisplayRange(NewDisplayLowBound, NewDisplayUpBound);
 		
-		constexpr float LeftBound = 0.f;
+		constexpr double LeftBound = 0.0;
 		MoveFocusPoint(LeftBound);
 	}
 	else if (NewPlayheadPosition >= FocusPointLockPosition)
@@ -234,7 +234,8 @@ void FSparseSampledSequenceTransportCoordinator::SetProgressRatio(const float Ne
 		else
 		{
 			// If playhead is far off screen, the NewPlayheadPosition could be > 1.0
-			MoveFocusPoint(NewPlayheadPosition);
+			constexpr double RightBound = 1.0;
+			MoveFocusPoint(FMath::Min(NewPlayheadPosition, RightBound));
 		}
 	}
 	else
