[AudioWidgets] Expose OnMouseCaptureBegin/OnMouseCaptureEnd delegates on all interactive widgets

- Allows undo transactions and other drag-lifecycle logic in the widget editor

#rnx
#Synced-CL 52401088

diff --git a/Engine/Plugins/Runtime/AudioWidgets/Source/AudioWidgets/Private/AudioMaterialSlate/AudioMaterialButton.cpp b/Engine/Plugins/Runtime/AudioWidgets/Source/AudioWidgets/Private/AudioMaterialSlate/AudioMaterialButton.cpp
--- a/Engine/Plugins/Runtime/AudioWidgets/Source/AudioWidgets/Private/AudioMaterialSlate/AudioMaterialButton.cpp
+++ b/Engine/Plugins/Runtime/AudioWidgets/Source/AudioWidgets/Private/AudioMaterialSlate/AudioMaterialButton.cpp
@@ -67,7 +67,9 @@ TSharedRef<SWidget> UAudioMaterialButton::RebuildWidget()
 		.Owner(this)
 		.AudioMaterialButtonStyle(&WidgetStyle)
 		.bIsPressedAttribute(bIsPressed)
-		.OnBooleanValueChanged(BIND_UOBJECT_DELEGATE(FOnBooleanValueChanged, HandleOnPressedValueChanged));
+		.OnBooleanValueChanged(BIND_UOBJECT_DELEGATE(FOnBooleanValueChanged, HandleOnPressedValueChanged))
+		.OnMouseCaptureBegin(BIND_UOBJECT_DELEGATE(FSimpleDelegate, HandleOnMouseCaptureBegin))
+		.OnMouseCaptureEnd(BIND_UOBJECT_DELEGATE(FSimpleDelegate, HandleOnMouseCaptureEnd));
 
 	return Button.ToSharedRef();
 }
@@ -78,4 +80,14 @@ void UAudioMaterialButton::HandleOnPressedValueChanged(bool InPressedSate)
 	OnButtonPressedChangedEvent.Broadcast(InPressedSate);
 }
 
+void UAudioMaterialButton::HandleOnMouseCaptureBegin()
+{
+	OnMouseCaptureBegin.Broadcast();
+}
+
+void UAudioMaterialButton::HandleOnMouseCaptureEnd()
+{
+	OnMouseCaptureEnd.Broadcast();
+}
+
 #undef LOCTEXT_NAMESPACE
diff --git a/Engine/Plugins/Runtime/AudioWidgets/Source/AudioWidgets/Private/AudioMaterialSlate/AudioMaterialKnob.cpp b/Engine/Plugins/Runtime/AudioWidgets/Source/AudioWidgets/Private/AudioMaterialSlate/AudioMaterialKnob.cpp
--- a/Engine/Plugins/Runtime/AudioWidgets/Source/AudioWidgets/Private/AudioMaterialSlate/AudioMaterialKnob.cpp
+++ b/Engine/Plugins/Runtime/AudioWidgets/Source/AudioWidgets/Private/AudioMaterialSlate/AudioMaterialKnob.cpp
@@ -6,6 +6,7 @@
 #include "AudioMaterialSlate/SAudioMaterialKnob.h"
 #include "AudioWidgetsStyle.h"
 
+
 #include UE_INLINE_GENERATED_CPP_BY_NAME(AudioMaterialKnob)
 
 #define LOCTEXT_NAMESPACE "AudioWidgets"
@@ -149,17 +150,29 @@ TSharedRef<SWidget> UAudioMaterialKnob::RebuildWidget()
 		.FineTuneSpeed(FineTuneSpeed)
 		.MouseUsesStep(bMouseUsesStep)
 		.StepSize(StepSize)
-		.OnFloatValueChanged(BIND_UOBJECT_DELEGATE(FOnFloatValueChanged, HandleOnKnobValueChanged));
-		
+		.OnFloatValueChanged(BIND_UOBJECT_DELEGATE(FOnFloatValueChanged, HandleOnKnobValueChanged))
+		.OnMouseCaptureBegin(BIND_UOBJECT_DELEGATE(FSimpleDelegate, HandleOnKnobCaptureBegin))
+		.OnMouseCaptureEnd(BIND_UOBJECT_DELEGATE(FSimpleDelegate, HandleOnKnobCaptureEnd));
+
 	return Knob.ToSharedRef();
 }
 
+void UAudioMaterialKnob::HandleOnKnobCaptureBegin()
+{
+	OnMouseCaptureBegin.Broadcast();
+}
+
+void UAudioMaterialKnob::HandleOnKnobCaptureEnd()
+{
+	OnMouseCaptureEnd.Broadcast();
+}
+
 void UAudioMaterialKnob::HandleOnKnobValueChanged(float InValue)
 {
 	if (Value != InValue)
 	{
 		Value = InValue;
-		OnKnobValueChanged.Broadcast(InValue);	
+		OnKnobValueChanged.Broadcast(InValue);
 	}
 }
 
diff --git a/Engine/Plugins/Runtime/AudioWidgets/Source/AudioWidgets/Private/AudioMaterialSlate/AudioMaterialSlider.cpp b/Engine/Plugins/Runtime/AudioWidgets/Source/AudioWidgets/Private/AudioMaterialSlate/AudioMaterialSlider.cpp
--- a/Engine/Plugins/Runtime/AudioWidgets/Source/AudioWidgets/Private/AudioMaterialSlate/AudioMaterialSlider.cpp
+++ b/Engine/Plugins/Runtime/AudioWidgets/Source/AudioWidgets/Private/AudioMaterialSlate/AudioMaterialSlider.cpp
@@ -151,7 +151,9 @@ TSharedRef<SWidget> UAudioMaterialSlider::RebuildWidget()
 		.MouseUsesStep(bMouseUsesStep)
 		.StepSize(StepSize)
 		.AudioMaterialSliderStyle(&WidgetStyle)
-		.OnValueChanged(BIND_UOBJECT_DELEGATE(FOnFloatValueChanged, HandleOnValueChanged));
+		.OnValueChanged(BIND_UOBJECT_DELEGATE(FOnFloatValueChanged, HandleOnValueChanged))
+		.OnMouseCaptureBegin(BIND_UOBJECT_DELEGATE(FSimpleDelegate, HandleOnMouseCaptureBegin))
+		.OnMouseCaptureEnd(BIND_UOBJECT_DELEGATE(FSimpleDelegate, HandleOnMouseCaptureEnd));
 
 	return Slider.ToSharedRef();
 }
@@ -161,8 +163,18 @@ void UAudioMaterialSlider::HandleOnValueChanged(float InValue)
 	if (Value != InValue)
 	{
 		Value = InValue;
-		OnValueChanged.Broadcast(InValue);	
+		OnValueChanged.Broadcast(InValue);
 	}
 }
 
+void UAudioMaterialSlider::HandleOnMouseCaptureBegin()
+{
+	OnMouseCaptureBegin.Broadcast();
+}
+
+void UAudioMaterialSlider::HandleOnMouseCaptureEnd()
+{
+	OnMouseCaptureEnd.Broadcast();
+}
+
 #undef LOCTEXT_NAMESPACE
diff --git a/Engine/Plugins/Runtime/AudioWidgets/Source/AudioWidgets/Private/AudioMaterialSlate/SAudioMaterialButton.cpp b/Engine/Plugins/Runtime/AudioWidgets/Source/AudioWidgets/Private/AudioMaterialSlate/SAudioMaterialButton.cpp
--- a/Engine/Plugins/Runtime/AudioWidgets/Source/AudioWidgets/Private/AudioMaterialSlate/SAudioMaterialButton.cpp
+++ b/Engine/Plugins/Runtime/AudioWidgets/Source/AudioWidgets/Private/AudioMaterialSlate/SAudioMaterialButton.cpp
@@ -13,6 +13,7 @@ void SAudioMaterialButton::Construct(const FArguments& InArgs)
 	AudioMaterialButtonStyle = InArgs._AudioMaterialButtonStyle;
 	bIsPressedAttribute = InArgs._bIsPressedAttribute;
 	OnBooleanValueChanged = InArgs._OnBooleanValueChanged;
+	OnMouseCaptureBegin = InArgs._OnMouseCaptureBegin;
 	OnMouseCaptureEnd = InArgs._OnMouseCaptureEnd;
 
 	ApplyNewMaterial();
@@ -104,8 +105,8 @@ FReply SAudioMaterialButton::OnMouseButtonDown(const FGeometry& MyGeometry, cons
 {
 	if ((MouseEvent.GetEffectingButton() == EKeys::LeftMouseButton))
 	{
+		OnMouseCaptureBegin.ExecuteIfBound();
 		CommitNewState(!bIsPressedAttribute.Get());
-
 		FReply Reply = FReply::Handled().CaptureMouse(SharedThis(this));
 		return Reply;
 	}
@@ -117,13 +118,18 @@ FReply SAudioMaterialButton::OnMouseButtonUp(const FGeometry& MyGeometry, const
 {
 	if ((MouseEvent.GetEffectingButton() == EKeys::LeftMouseButton) && this->HasMouseCapture())
 	{
-		OnMouseCaptureEnd.ExecuteIfBound();
 		return FReply::Handled().ReleaseMouseCapture();
 	}
 
 	return FReply::Unhandled();
 }
 
+void SAudioMaterialButton::OnMouseCaptureLost(const FCaptureLostEvent& CaptureLostEvent)
+{
+	OnMouseCaptureEnd.ExecuteIfBound();
+	SLeafWidget::OnMouseCaptureLost(CaptureLostEvent);
+}
+
 void SAudioMaterialButton::CommitNewState(bool InPressedState)
 {
 	if (bIsPressedAttribute.Get() != InPressedState)
diff --git a/Engine/Plugins/Runtime/AudioWidgets/Source/AudioWidgets/Private/AudioMaterialSlate/SAudioMaterialKnob.cpp b/Engine/Plugins/Runtime/AudioWidgets/Source/AudioWidgets/Private/AudioMaterialSlate/SAudioMaterialKnob.cpp
--- a/Engine/Plugins/Runtime/AudioWidgets/Source/AudioWidgets/Private/AudioMaterialSlate/SAudioMaterialKnob.cpp
+++ b/Engine/Plugins/Runtime/AudioWidgets/Source/AudioWidgets/Private/AudioMaterialSlate/SAudioMaterialKnob.cpp
@@ -22,6 +22,7 @@ void SAudioMaterialKnob::Construct(const FArguments& InArgs)
 	AudioMaterialKnobStyle = InArgs._AudioMaterialKnobStyle;
 
 	OnValueChanged = InArgs._OnFloatValueChanged;
+	OnValueCommitted = InArgs._OnValueCommitted;
 	OnMouseCaptureBegin = InArgs._OnMouseCaptureBegin;
 	OnMouseCaptureEnd = InArgs._OnMouseCaptureEnd;
 
@@ -222,14 +223,20 @@ FReply SAudioMaterialKnob::OnMouseButtonUp(const FGeometry& MyGeometry, const FP
 {
 	if ((MouseEvent.GetEffectingButton() == EKeys::LeftMouseButton) && this->HasMouseCapture())
 	{
-		SetCursor(CachedCursor);
-		OnMouseCaptureEnd.ExecuteIfBound();
 		return FReply::Handled().ReleaseMouseCapture();
 	}
 
 	return FReply::Unhandled();
 }
 
+void SAudioMaterialKnob::OnMouseCaptureLost(const FCaptureLostEvent& CaptureLostEvent)
+{
+	SetCursor(CachedCursor);
+	OnValueCommitted.ExecuteIfBound(ValueAttribute.Get());
+	OnMouseCaptureEnd.ExecuteIfBound();
+	SLeafWidget::OnMouseCaptureLost(CaptureLostEvent);
+}
+
 FReply SAudioMaterialKnob::OnKeyDown(const FGeometry& MyGeometry, const FKeyEvent& InKeyEvent)
 {
 	if (InKeyEvent.GetKey() == EKeys::LeftShift)
diff --git a/Engine/Plugins/Runtime/AudioWidgets/Source/AudioWidgets/Private/AudioMaterialSlate/SAudioMaterialSlider.cpp b/Engine/Plugins/Runtime/AudioWidgets/Source/AudioWidgets/Private/AudioMaterialSlate/SAudioMaterialSlider.cpp
--- a/Engine/Plugins/Runtime/AudioWidgets/Source/AudioWidgets/Private/AudioMaterialSlate/SAudioMaterialSlider.cpp
+++ b/Engine/Plugins/Runtime/AudioWidgets/Source/AudioWidgets/Private/AudioMaterialSlate/SAudioMaterialSlider.cpp
@@ -26,6 +26,8 @@ void SAudioMaterialSlider::Construct(const FArguments& InArgs)
 
 	OnValueChanged = InArgs._OnValueChanged;
 	OnValueCommitted = InArgs._OnValueCommitted;
+	OnMouseCaptureBegin = InArgs._OnMouseCaptureBegin;
+	OnMouseCaptureEnd = InArgs._OnMouseCaptureEnd;
 
 	ApplyNewMaterial();
 
@@ -171,6 +173,7 @@ FReply SAudioMaterialSlider::OnMouseButtonDown(const FGeometry& MyGeometry, cons
 
 		MouseDownStartPosition = MouseEvent.GetScreenSpacePosition();
 		MouseDownValue = ValueAttribute.Get();
+		OnMouseCaptureBegin.ExecuteIfBound();
 		return FReply::Handled().CaptureMouse(SharedThis(this));
 	}
 
@@ -181,7 +184,6 @@ FReply SAudioMaterialSlider::OnMouseButtonUp(const FGeometry& MyGeometry, const
 {
 	if ((MouseEvent.GetEffectingButton() == EKeys::LeftMouseButton) && HasMouseCapture())
 	{
-		SetCursor(CachedCursor);
 		return FReply::Handled().ReleaseMouseCapture();
 	}
 
@@ -218,7 +220,9 @@ FReply SAudioMaterialSlider::OnMouseMove(const FGeometry& MyGeometry, const FPoi
 
 void SAudioMaterialSlider::OnMouseCaptureLost(const FCaptureLostEvent& CaptureLostEvent)
 {
+	SetCursor(CachedCursor);
 	OnValueCommitted.ExecuteIfBound(ValueAttribute.Get());
+	OnMouseCaptureEnd.ExecuteIfBound();
 	SLeafWidget::OnMouseCaptureLost(CaptureLostEvent);
 }
 
diff --git a/Engine/Plugins/Runtime/AudioWidgets/Source/AudioWidgets/Private/AudioRadialSlider.cpp b/Engine/Plugins/Runtime/AudioWidgets/Source/AudioWidgets/Private/AudioRadialSlider.cpp
--- a/Engine/Plugins/Runtime/AudioWidgets/Source/AudioWidgets/Private/AudioRadialSlider.cpp
+++ b/Engine/Plugins/Runtime/AudioWidgets/Source/AudioWidgets/Private/AudioRadialSlider.cpp
@@ -212,6 +212,16 @@ void UAudioRadialSlider::HandleOnValueChanged(float InValue)
 	OnValueChanged.Broadcast(InValue);
 }
 
+void UAudioRadialSlider::HandleOnMouseCaptureBegin()
+{
+	OnMouseCaptureBegin.Broadcast();
+}
+
+void UAudioRadialSlider::HandleOnMouseCaptureEnd()
+{
+	OnMouseCaptureEnd.Broadcast();
+}
+
 #if WITH_EDITOR
 const FText UAudioRadialSlider::GetPaletteCategory()
 {
@@ -230,7 +240,9 @@ TSharedRef<SWidget> UAudioRadialSlider::RebuildWidget()
 {
 	MyAudioRadialSlider = SNew(SAudioRadialSlider)
 		.Style(&WidgetStyle)
-		.OnValueChanged(BIND_UOBJECT_DELEGATE(FOnFloatValueChanged, HandleOnValueChanged));
+		.OnValueChanged(BIND_UOBJECT_DELEGATE(FOnFloatValueChanged, HandleOnValueChanged))
+		.OnMouseCaptureBegin(BIND_UOBJECT_DELEGATE(FSimpleDelegate, HandleOnMouseCaptureBegin))
+		.OnMouseCaptureEnd(BIND_UOBJECT_DELEGATE(FSimpleDelegate, HandleOnMouseCaptureEnd));
 
 	return MyAudioRadialSlider.ToSharedRef();
 }
@@ -247,7 +259,9 @@ TSharedRef<SWidget> UAudioVolumeRadialSlider::RebuildWidget()
 {
 	MyAudioRadialSlider = SNew(SAudioVolumeRadialSlider)
 		.Style(&WidgetStyle)
-		.OnValueChanged(BIND_UOBJECT_DELEGATE(FOnFloatValueChanged, HandleOnValueChanged));
+		.OnValueChanged(BIND_UOBJECT_DELEGATE(FOnFloatValueChanged, HandleOnValueChanged))
+		.OnMouseCaptureBegin(BIND_UOBJECT_DELEGATE(FSimpleDelegate, HandleOnMouseCaptureBegin))
+		.OnMouseCaptureEnd(BIND_UOBJECT_DELEGATE(FSimpleDelegate, HandleOnMouseCaptureEnd));
 
 	return MyAudioRadialSlider.ToSharedRef();
 }
@@ -264,7 +278,9 @@ TSharedRef<SWidget> UAudioFrequencyRadialSlider::RebuildWidget()
 {
 	MyAudioRadialSlider = SNew(SAudioFrequencyRadialSlider)
 		.Style(&WidgetStyle)
-		.OnValueChanged(BIND_UOBJECT_DELEGATE(FOnFloatValueChanged, HandleOnValueChanged));
+		.OnValueChanged(BIND_UOBJECT_DELEGATE(FOnFloatValueChanged, HandleOnValueChanged))
+		.OnMouseCaptureBegin(BIND_UOBJECT_DELEGATE(FSimpleDelegate, HandleOnMouseCaptureBegin))
+		.OnMouseCaptureEnd(BIND_UOBJECT_DELEGATE(FSimpleDelegate, HandleOnMouseCaptureEnd));
 
 	return MyAudioRadialSlider.ToSharedRef();
 }
diff --git a/Engine/Plugins/Runtime/AudioWidgets/Source/AudioWidgets/Private/AudioSlider.cpp b/Engine/Plugins/Runtime/AudioWidgets/Source/AudioWidgets/Private/AudioSlider.cpp
--- a/Engine/Plugins/Runtime/AudioWidgets/Source/AudioWidgets/Private/AudioSlider.cpp
+++ b/Engine/Plugins/Runtime/AudioWidgets/Source/AudioWidgets/Private/AudioSlider.cpp
@@ -106,6 +106,16 @@ void UAudioSliderBase::HandleOnValueChanged(float InValue)
 	OnValueChanged.Broadcast(InValue);
 }
 
+void UAudioSliderBase::HandleOnMouseCaptureBegin()
+{
+	OnMouseCaptureBegin.Broadcast();
+}
+
+void UAudioSliderBase::HandleOnMouseCaptureEnd()
+{
+	OnMouseCaptureEnd.Broadcast();
+}
+
 void UAudioSliderBase::SetTextLabelBackgroundColor(FSlateColor InColor)
 {
 	TextLabelBackgroundColor = InColor.GetSpecifiedColor();
@@ -215,7 +225,9 @@ TSharedRef<SWidget> UAudioSliderBase::RebuildWidget()
 {
 	MyAudioSlider = SNew(SAudioSliderBase)
 		.Style(&WidgetStyle)
-		.OnValueChanged(BIND_UOBJECT_DELEGATE(FOnFloatValueChanged, HandleOnValueChanged));
+		.OnValueChanged(BIND_UOBJECT_DELEGATE(FOnFloatValueChanged, HandleOnValueChanged))
+		.OnMouseCaptureBegin(BIND_UOBJECT_DELEGATE(FSimpleDelegate, HandleOnMouseCaptureBegin))
+		.OnMouseCaptureEnd(BIND_UOBJECT_DELEGATE(FSimpleDelegate, HandleOnMouseCaptureEnd));
 
 	return MyAudioSlider.ToSharedRef();
 }
@@ -247,7 +259,9 @@ TSharedRef<SWidget> UAudioSlider::RebuildWidget()
 {
 	MyAudioSlider = SNew(SAudioSlider)
 		.Style(&WidgetStyle)
-		.OnValueChanged(BIND_UOBJECT_DELEGATE(FOnFloatValueChanged, HandleOnValueChanged));
+		.OnValueChanged(BIND_UOBJECT_DELEGATE(FOnFloatValueChanged, HandleOnValueChanged))
+		.OnMouseCaptureBegin(BIND_UOBJECT_DELEGATE(FSimpleDelegate, HandleOnMouseCaptureBegin))
+		.OnMouseCaptureEnd(BIND_UOBJECT_DELEGATE(FSimpleDelegate, HandleOnMouseCaptureEnd));
 
 	return MyAudioSlider.ToSharedRef();
 }
@@ -262,7 +276,9 @@ UAudioVolumeSlider::UAudioVolumeSlider(const FObjectInitializer& ObjectInitializ
 TSharedRef<SWidget> UAudioVolumeSlider::RebuildWidget()
 {
 	MyAudioSlider = SNew(SAudioVolumeSlider)
-		.OnValueChanged(BIND_UOBJECT_DELEGATE(FOnFloatValueChanged, HandleOnValueChanged));
+		.OnValueChanged(BIND_UOBJECT_DELEGATE(FOnFloatValueChanged, HandleOnValueChanged))
+		.OnMouseCaptureBegin(BIND_UOBJECT_DELEGATE(FSimpleDelegate, HandleOnMouseCaptureBegin))
+		.OnMouseCaptureEnd(BIND_UOBJECT_DELEGATE(FSimpleDelegate, HandleOnMouseCaptureEnd));
 	LinToOutputCurve = StaticCastSharedPtr<SAudioSlider>(MyAudioSlider)->GetLinToOutputCurve();
 	OutputToLinCurve = StaticCastSharedPtr<SAudioSlider>(MyAudioSlider)->GetOutputToLinCurve();
 
@@ -281,7 +297,9 @@ TSharedRef<SWidget> UAudioFrequencySlider::RebuildWidget()
 {
 	MyAudioSlider = SNew(SAudioFrequencySlider)
 		.Style(&WidgetStyle)
-		.OnValueChanged(BIND_UOBJECT_DELEGATE(FOnFloatValueChanged, HandleOnValueChanged));
+		.OnValueChanged(BIND_UOBJECT_DELEGATE(FOnFloatValueChanged, HandleOnValueChanged))
+		.OnMouseCaptureBegin(BIND_UOBJECT_DELEGATE(FSimpleDelegate, HandleOnMouseCaptureBegin))
+		.OnMouseCaptureEnd(BIND_UOBJECT_DELEGATE(FSimpleDelegate, HandleOnMouseCaptureEnd));
 	StaticCastSharedPtr<SAudioFrequencySlider>(MyAudioSlider)->SetOutputRange(OutputRange);
 
 	return MyAudioSlider.ToSharedRef();
diff --git a/Engine/Plugins/Runtime/AudioWidgets/Source/AudioWidgets/Private/SAudioSlider.cpp b/Engine/Plugins/Runtime/AudioWidgets/Source/AudioWidgets/Private/SAudioSlider.cpp
--- a/Engine/Plugins/Runtime/AudioWidgets/Source/AudioWidgets/Private/SAudioSlider.cpp
+++ b/Engine/Plugins/Runtime/AudioWidgets/Source/AudioWidgets/Private/SAudioSlider.cpp
@@ -23,6 +23,8 @@ void SAudioSliderBase::Construct(const SAudioSliderBase::FArguments& InArgs)
 	Style = InArgs._Style;
 	OnValueChanged = InArgs._OnValueChanged;
 	OnValueCommitted = InArgs._OnValueCommitted;
+	OnMouseCaptureBegin = InArgs._OnMouseCaptureBegin;
+	OnMouseCaptureEnd = InArgs._OnMouseCaptureEnd;
 	SliderValueAttribute = InArgs._SliderValue;
 	SliderBackgroundColor = InArgs._SliderBackgroundColor;
 	SliderBarColor = InArgs._SliderBarColor;
@@ -61,9 +63,14 @@ void SAudioSliderBase::Construct(const SAudioSliderBase::FArguments& InArgs)
 			const float OutputValue = GetOutputValueForText(Value);
 			Label->SetValueText(OutputValue);
 		})
+		.OnMouseCaptureBegin_Lambda([this]()
+		{
+			OnMouseCaptureBegin.ExecuteIfBound();
+		})
 		.OnMouseCaptureEnd_Lambda([this]()
 		{
 			OnValueCommitted.ExecuteIfBound(SliderValueAttribute.Get());
+			OnMouseCaptureEnd.ExecuteIfBound();
 		});
 
 	// Text label
diff --git a/Engine/Plugins/Runtime/AudioWidgets/Source/AudioWidgets/Public/AudioMaterialSlate/AudioMaterialButton.h b/Engine/Plugins/Runtime/AudioWidgets/Source/AudioWidgets/Public/AudioMaterialSlate/AudioMaterialButton.h
--- a/Engine/Plugins/Runtime/AudioWidgets/Source/AudioWidgets/Public/AudioMaterialSlate/AudioMaterialButton.h
+++ b/Engine/Plugins/Runtime/AudioWidgets/Source/AudioWidgets/Public/AudioMaterialSlate/AudioMaterialButton.h
@@ -12,6 +12,8 @@
 class SAudioMaterialButton;
 struct FAudioMaterialButtonStyle;
 
+DECLARE_DYNAMIC_MULTICAST_DELEGATE(FOnButtonMouseCaptureBeginEvent);
+DECLARE_DYNAMIC_MULTICAST_DELEGATE(FOnButtonMouseCaptureEndEvent);
 DECLARE_DYNAMIC_MULTICAST_DELEGATE_OneParam(FOnButtonPressedChangedEvent, bool, bIsPressed);
 
 
@@ -62,6 +64,14 @@ public:
 	UPROPERTY(BlueprintAssignable, Category = "Widget Event")
 	FOnButtonPressedChangedEvent OnButtonPressedChangedEvent;
 
+	/** Called when mouse capture begins. */
+	UPROPERTY(BlueprintAssignable, Category = "Widget Event")
+	FOnButtonMouseCaptureBeginEvent OnMouseCaptureBegin;
+
+	/** Called when mouse capture ends (mouse release or any other capture loss). */
+	UPROPERTY(BlueprintAssignable, Category = "Widget Event")
+	FOnButtonMouseCaptureEndEvent OnMouseCaptureEnd;
+
 protected:
 
 	// UWidget
@@ -69,6 +79,8 @@ protected:
 	// End of UWidget
 
 	UE_API void HandleOnPressedValueChanged(bool InPressedSate);
+	UE_API void HandleOnMouseCaptureBegin();
+	UE_API void HandleOnMouseCaptureEnd();
 
 private:
 
diff --git a/Engine/Plugins/Runtime/AudioWidgets/Source/AudioWidgets/Public/AudioMaterialSlate/AudioMaterialKnob.h b/Engine/Plugins/Runtime/AudioWidgets/Source/AudioWidgets/Public/AudioMaterialSlate/AudioMaterialKnob.h
--- a/Engine/Plugins/Runtime/AudioWidgets/Source/AudioWidgets/Public/AudioMaterialSlate/AudioMaterialKnob.h
+++ b/Engine/Plugins/Runtime/AudioWidgets/Source/AudioWidgets/Public/AudioMaterialSlate/AudioMaterialKnob.h
@@ -11,6 +11,8 @@
 
 class SAudioMaterialKnob;
 
+DECLARE_DYNAMIC_MULTICAST_DELEGATE(FOnKnobMouseCaptureBeginEvent);
+DECLARE_DYNAMIC_MULTICAST_DELEGATE(FOnKnobMouseCaptureEndEvent);
 DECLARE_DYNAMIC_MULTICAST_DELEGATE_OneParam(FOnKnobValueChangedEvent, float, Value);
 
 
@@ -100,6 +102,14 @@ public:
 	UPROPERTY(BlueprintAssignable, Category = "Widget Event")
 	FOnKnobValueChangedEvent OnKnobValueChanged;
 
+	/** Called when mouse capture begins. */
+	UPROPERTY(BlueprintAssignable, Category = "Widget Event")
+	FOnKnobMouseCaptureBeginEvent OnMouseCaptureBegin;
+
+	/** Called when mouse capture ends (mouse release or any other capture loss). */
+	UPROPERTY(BlueprintAssignable, Category = "Widget Event")
+	FOnKnobMouseCaptureEndEvent OnMouseCaptureEnd;
+
 protected:
 
 	// UWidget
@@ -136,6 +146,9 @@ private:
 
 private:
 
+	void HandleOnKnobCaptureBegin();
+	void HandleOnKnobCaptureEnd();
+
 	/** Native Slate Widget */
 	TSharedPtr<SAudioMaterialKnob> Knob;
 
diff --git a/Engine/Plugins/Runtime/AudioWidgets/Source/AudioWidgets/Public/AudioMaterialSlate/AudioMaterialSlider.h b/Engine/Plugins/Runtime/AudioWidgets/Source/AudioWidgets/Public/AudioMaterialSlate/AudioMaterialSlider.h
--- a/Engine/Plugins/Runtime/AudioWidgets/Source/AudioWidgets/Public/AudioMaterialSlate/AudioMaterialSlider.h
+++ b/Engine/Plugins/Runtime/AudioWidgets/Source/AudioWidgets/Public/AudioMaterialSlate/AudioMaterialSlider.h
@@ -11,6 +11,8 @@
 
 class SAudioMaterialSlider;
 
+DECLARE_DYNAMIC_MULTICAST_DELEGATE(FOnSliderMouseCaptureBeginEvent);
+DECLARE_DYNAMIC_MULTICAST_DELEGATE(FOnSliderMouseCaptureEndEvent);
 DECLARE_DYNAMIC_MULTICAST_DELEGATE_OneParam(FOnSliderFloatValueChangedEvent, float, Value);
 
 /**
@@ -100,6 +102,14 @@ public:
 	UPROPERTY(BlueprintAssignable, FieldNotify, Category = "Widget Event")
 	FOnSliderFloatValueChangedEvent OnValueChanged;
 
+	/** Called when mouse capture begins. */
+	UPROPERTY(BlueprintAssignable, Category = "Widget Event")
+	FOnSliderMouseCaptureBeginEvent OnMouseCaptureBegin;
+
+	/** Called when mouse capture ends (mouse release or any other capture loss). */
+	UPROPERTY(BlueprintAssignable, Category = "Widget Event")
+	FOnSliderMouseCaptureEndEvent OnMouseCaptureEnd;
+
 protected:
 
 	// UWidget
@@ -107,6 +117,8 @@ protected:
 	// End of UWidget
 
 	UE_API void HandleOnValueChanged(float InValue);
+	UE_API void HandleOnMouseCaptureBegin();
+	UE_API void HandleOnMouseCaptureEnd();
 
 private:
 
diff --git a/Engine/Plugins/Runtime/AudioWidgets/Source/AudioWidgets/Public/AudioMaterialSlate/SAudioMaterialButton.h b/Engine/Plugins/Runtime/AudioWidgets/Source/AudioWidgets/Public/AudioMaterialSlate/SAudioMaterialButton.h
--- a/Engine/Plugins/Runtime/AudioWidgets/Source/AudioWidgets/Public/AudioMaterialSlate/SAudioMaterialButton.h
+++ b/Engine/Plugins/Runtime/AudioWidgets/Source/AudioWidgets/Public/AudioMaterialSlate/SAudioMaterialButton.h
@@ -37,7 +37,10 @@ public:
 	/** Called when the button's state changes. */
 	SLATE_EVENT(FOnBooleanValueChanged, OnBooleanValueChanged)
 
-	/** Invoked when the mouse is released and a capture ends. */
+	/** Invoked when the mouse is pressed and a capture begins. */
+	SLATE_EVENT(FSimpleDelegate, OnMouseCaptureBegin)
+
+	/** Invoked when mouse capture ends (mouse release or any other capture loss). */
 	SLATE_EVENT(FSimpleDelegate, OnMouseCaptureEnd)
 
 	SLATE_END_ARGS()
@@ -58,7 +61,10 @@ public:
 
 	FOnBooleanValueChanged OnBooleanValueChanged;
 
-	// Holds a delegate that is executed when the mouse is let up and a capture ends.
+	// Holds a delegate that is executed when the mouse is pressed and a capture begins.
+	FSimpleDelegate OnMouseCaptureBegin;
+
+	// Holds a delegate that is executed when mouse capture ends.
 	FSimpleDelegate OnMouseCaptureEnd;
 
 protected:
@@ -67,7 +73,8 @@ protected:
 	UE_API virtual int32 OnPaint(const FPaintArgs& Args, const FGeometry& AllottedGeometry, const FSlateRect& MyCullingRect, FSlateWindowElementList& OutDrawElements, int32 LayerId, const FWidgetStyle& InWidgetStyle, bool bParentEnabled) const override;
 	UE_API virtual FVector2D ComputeDesiredSize(float) const override;
 	UE_API virtual FReply OnMouseButtonDown(const FGeometry& MyGeometry, const FPointerEvent& MouseEvent) override;
-	UE_API virtual FReply OnMouseButtonUp(const FGeometry& MyGeometry, const FPointerEvent& MouseEvent)override;
+	UE_API virtual FReply OnMouseButtonUp(const FGeometry& MyGeometry, const FPointerEvent& MouseEvent) override;
+	UE_API virtual void OnMouseCaptureLost(const FCaptureLostEvent& CaptureLostEvent) override;
 
 private:
 
diff --git a/Engine/Plugins/Runtime/AudioWidgets/Source/AudioWidgets/Public/AudioMaterialSlate/SAudioMaterialKnob.h b/Engine/Plugins/Runtime/AudioWidgets/Source/AudioWidgets/Public/AudioMaterialSlate/SAudioMaterialKnob.h
--- a/Engine/Plugins/Runtime/AudioWidgets/Source/AudioWidgets/Public/AudioMaterialSlate/SAudioMaterialKnob.h
+++ b/Engine/Plugins/Runtime/AudioWidgets/Source/AudioWidgets/Public/AudioMaterialSlate/SAudioMaterialKnob.h
@@ -60,10 +60,13 @@ public:
 	/** Called when the knob's state changes. */
 	SLATE_EVENT(FOnFloatValueChanged, OnFloatValueChanged)
 
+	/** Called when the value is committed (mouse capture ends) */
+	SLATE_EVENT(FOnFloatValueChanged, OnValueCommitted)
+
 	/** Invoked when the mouse is pressed and a capture begins. */
 	SLATE_EVENT(FSimpleDelegate, OnMouseCaptureBegin)
 
-	/** Invoked when the mouse is released and a capture ends. */
+	/** Invoked when mouse capture ends (mouse release or any other capture loss). */
 	SLATE_EVENT(FSimpleDelegate, OnMouseCaptureEnd)
 
 	SLATE_END_ARGS()
@@ -109,10 +112,13 @@ public:
 	// Holds a delegate that is executed when the knob's value changes.
 	FOnFloatValueChanged OnValueChanged;
 
+	// Holds a delegate that is executed when the knob's value is committed (mouse capture ends).
+	FOnFloatValueChanged OnValueCommitted;
+
 	// Holds a delegate that is executed when the mouse is pressed and a capture begins.
 	FSimpleDelegate OnMouseCaptureBegin;
 
-	// Holds a delegate that is executed when the mouse is let up and a capture ends.
+	// Holds a delegate that is executed when mouse capture ends.
 	FSimpleDelegate OnMouseCaptureEnd;
 
 protected:
@@ -122,7 +128,8 @@ protected:
 	UE_API virtual FVector2D ComputeDesiredSize(float) const override;
 	UE_API virtual FReply OnMouseMove(const FGeometry& MyGeometry, const FPointerEvent& MouseEvent) override;
 	UE_API virtual FReply OnMouseButtonDown(const FGeometry& MyGeometry, const FPointerEvent& MouseEvent) override;
-	UE_API virtual FReply OnMouseButtonUp(const FGeometry& MyGeometry, const FPointerEvent& MouseEvent)override;
+	UE_API virtual FReply OnMouseButtonUp(const FGeometry& MyGeometry, const FPointerEvent& MouseEvent) override;
+	UE_API virtual void OnMouseCaptureLost(const FCaptureLostEvent& CaptureLostEvent) override;
 	UE_API virtual FReply OnKeyDown(const FGeometry& MyGeometry, const FKeyEvent& InKeyEvent) override;
 	UE_API virtual FReply OnKeyUp(const FGeometry& MyGeometry, const FKeyEvent& InKeyEvent) override;
 	UE_API virtual bool SupportsKeyboardFocus() const override;
diff --git a/Engine/Plugins/Runtime/AudioWidgets/Source/AudioWidgets/Public/AudioMaterialSlate/SAudioMaterialSlider.h b/Engine/Plugins/Runtime/AudioWidgets/Source/AudioWidgets/Public/AudioMaterialSlate/SAudioMaterialSlider.h
--- a/Engine/Plugins/Runtime/AudioWidgets/Source/AudioWidgets/Public/AudioMaterialSlate/SAudioMaterialSlider.h
+++ b/Engine/Plugins/Runtime/AudioWidgets/Source/AudioWidgets/Public/AudioMaterialSlate/SAudioMaterialSlider.h
@@ -60,10 +60,16 @@ public:
 
 	/** Called when the value is changed by the slider. */
 	SLATE_EVENT(FOnFloatValueChanged, OnValueChanged)
-	
+
 	/** Called when the value is committed (mouse capture ends) */
 	SLATE_EVENT(FOnFloatValueChanged, OnValueCommitted)
 
+	/** Invoked when the mouse is pressed and a capture begins. */
+	SLATE_EVENT(FSimpleDelegate, OnMouseCaptureBegin)
+
+	/** Invoked when mouse capture ends (mouse release or any other capture loss). */
+	SLATE_EVENT(FSimpleDelegate, OnMouseCaptureEnd)
+
 	SLATE_END_ARGS()
 
 	/**
@@ -117,10 +123,16 @@ public:
 
 	// Holds a delegate that is executed when the slider's value changed.
 	FOnFloatValueChanged OnValueChanged;
-	
+
 	// Holds a delegate that is executed when the slider's value is committed (mouse capture ends).
 	FOnFloatValueChanged OnValueCommitted;
 
+	// Holds a delegate that is executed when the mouse is pressed and a capture begins.
+	FSimpleDelegate OnMouseCaptureBegin;
+
+	// Holds a delegate that is executed when mouse capture ends.
+	FSimpleDelegate OnMouseCaptureEnd;
+
 private:
 
 	/**
diff --git a/Engine/Plugins/Runtime/AudioWidgets/Source/AudioWidgets/Public/AudioRadialSlider.h b/Engine/Plugins/Runtime/AudioWidgets/Source/AudioWidgets/Public/AudioRadialSlider.h
--- a/Engine/Plugins/Runtime/AudioWidgets/Source/AudioWidgets/Public/AudioRadialSlider.h
+++ b/Engine/Plugins/Runtime/AudioWidgets/Source/AudioWidgets/Public/AudioRadialSlider.h
@@ -11,6 +11,8 @@
 class SAudioRadialSlider;
 enum EAudioRadialSliderLayout : int;
 
+DECLARE_DYNAMIC_MULTICAST_DELEGATE(FOnAudioRadialSliderMouseCaptureBeginEvent);
+DECLARE_DYNAMIC_MULTICAST_DELEGATE(FOnAudioRadialSliderMouseCaptureEndEvent);
 DECLARE_DYNAMIC_MULTICAST_DELEGATE_OneParam(FOnAudioRadialSliderValueChangedEvent, float, Value);
 
 /**
@@ -86,6 +88,14 @@ public:
 	UPROPERTY(BlueprintAssignable, Category = "Widget Event")
 	FOnAudioRadialSliderValueChangedEvent OnValueChanged;
 
+	/** Called when a mouse drag begins. */
+	UPROPERTY(BlueprintAssignable, Category = "Widget Event")
+	FOnAudioRadialSliderMouseCaptureBeginEvent OnMouseCaptureBegin;
+
+	/** Called when a mouse drag ends. */
+	UPROPERTY(BlueprintAssignable, Category = "Widget Event")
+	FOnAudioRadialSliderMouseCaptureEndEvent OnMouseCaptureEnd;
+
 public: 
 	/** Get output value from normalized linear (0 - 1) based on internal lin to output mapping. */
 	UFUNCTION(BlueprintCallable, Category = "Audio Widgets| Audio Radial Slider")
@@ -166,6 +176,8 @@ public:
 	TSharedPtr<SAudioRadialSlider> MyAudioRadialSlider;
 
 	UE_API void HandleOnValueChanged(float InValue);
+	UE_API void HandleOnMouseCaptureBegin();
+	UE_API void HandleOnMouseCaptureEnd();
 
 	// UWidget interface
 	UE_API virtual TSharedRef<SWidget> RebuildWidget();
diff --git a/Engine/Plugins/Runtime/AudioWidgets/Source/AudioWidgets/Public/AudioSlider.h b/Engine/Plugins/Runtime/AudioWidgets/Source/AudioWidgets/Public/AudioSlider.h
--- a/Engine/Plugins/Runtime/AudioWidgets/Source/AudioWidgets/Public/AudioSlider.h
+++ b/Engine/Plugins/Runtime/AudioWidgets/Source/AudioWidgets/Public/AudioSlider.h
@@ -151,6 +151,14 @@ public:
 	UPROPERTY(BlueprintAssignable, Category = "Widget Event")
 	FOnFloatValueChangedEvent OnValueChanged;
 
+	/** Called when a mouse drag begins. */
+	UPROPERTY(BlueprintAssignable, Category = "Widget Event")
+	FOnMouseCaptureBeginEvent OnMouseCaptureBegin;
+
+	/** Called when a mouse drag ends. */
+	UPROPERTY(BlueprintAssignable, Category = "Widget Event")
+	FOnMouseCaptureEndEvent OnMouseCaptureEnd;
+
 	// UWidget interface
 	UE_API virtual void SynchronizeProperties() override;
 	// End of UWidget interface
@@ -170,6 +178,8 @@ protected:
 	TSharedPtr<SAudioSliderBase> MyAudioSlider;
 
 	UE_API void HandleOnValueChanged(float InValue);
+	UE_API void HandleOnMouseCaptureBegin();
+	UE_API void HandleOnMouseCaptureEnd();
 
 	// UWidget interface
 	UE_API virtual TSharedRef<SWidget> RebuildWidget();
diff --git a/Engine/Plugins/Runtime/AudioWidgets/Source/AudioWidgets/Public/SAudioSlider.h b/Engine/Plugins/Runtime/AudioWidgets/Source/AudioWidgets/Public/SAudioSlider.h
--- a/Engine/Plugins/Runtime/AudioWidgets/Source/AudioWidgets/Public/SAudioSlider.h
+++ b/Engine/Plugins/Runtime/AudioWidgets/Source/AudioWidgets/Public/SAudioSlider.h
@@ -73,6 +73,12 @@ public:
 		/** Called when the value is committed by mouse capture ending */
 		SLATE_EVENT(FOnFloatValueChanged, OnValueCommitted)
 
+		/** Invoked when the mouse is pressed and a capture begins. */
+		SLATE_EVENT(FSimpleDelegate, OnMouseCaptureBegin)
+
+		/** Invoked when mouse capture ends (mouse release or any other capture loss). */
+		SLATE_EVENT(FSimpleDelegate, OnMouseCaptureEnd)
+
 	SLATE_END_ARGS()
 
 	UE_API SAudioSliderBase();
@@ -84,6 +90,12 @@ public:
 	// Holds a delegate that is executed when the slider's value is committed (mouse capture ends).
 	FOnFloatValueChanged OnValueCommitted;
 
+	// Holds a delegate that is executed when the mouse is pressed and a capture begins.
+	FSimpleDelegate OnMouseCaptureBegin;
+
+	// Holds a delegate that is executed when mouse capture ends.
+	FSimpleDelegate OnMouseCaptureEnd;
+
 	/**
 	 * Construct the widget.
 	 *
